Add Nearest Neighbor Editor to MLDeformer

diff --git a/Engine/Plugins/Animation/MLDeformer/MLDeformerFramework/Source/MLDeformerFramework/Public/MLDeformerInputInfo.h b/Engine/Plugins/Animation/MLDeformer/MLDeformerFramework/Source/MLDeformerFramework/Public/MLDeformerInputInfo.h
--- a/Engine/Plugins/Animation/MLDeformer/MLDeformerFramework/Source/MLDeformerFramework/Public/MLDeformerInputInfo.h
+++ b/Engine/Plugins/Animation/MLDeformer/MLDeformerFramework/Source/MLDeformerFramework/Public/MLDeformerInputInfo.h
@@ -141,7 +141,7 @@ public:
 	 * @param SkelMeshComponent The skeletal mesh component to sample from.
 	 * @param OutValues The array to write the values to. This array will be reset/resized by this method.
 	 */
-	void ExtractCurveValues(USkeletalMeshComponent* SkelMeshComponent, TArray<float>& OutValues) const;
+	virtual void ExtractCurveValues(USkeletalMeshComponent* SkelMeshComponent, TArray<float>& OutValues) const;
 
 	/**
 	 * Extract bone space rotations, as a float array.
@@ -150,7 +150,7 @@ public:
 	 * @param SkelMeshComponent The skeletal mesh component to sample from.
 	 * @param OutRotations The output rotation values. This array will be resized internally.
 	 */
-	void ExtractBoneRotations(USkeletalMeshComponent* SkelMeshComponent, TArray<float>& OutRotations) const;
+	virtual void ExtractBoneRotations(USkeletalMeshComponent* SkelMeshComponent, TArray<float>& OutRotations) const;
 
 	/**
 	 * Get the number of imported vertices in the base mesh, which is the linear skinned skeletal mesh.
diff --git a/Engine/Plugins/Animation/MLDeformer/MLDeformerFramework/Source/MLDeformerFramework/Public/MLDeformerModelInstance.h b/Engine/Plugins/Animation/MLDeformer/MLDeformerFramework/Source/MLDeformerFramework/Public/MLDeformerModelInstance.h
--- a/Engine/Plugins/Animation/MLDeformer/MLDeformerFramework/Source/MLDeformerFramework/Public/MLDeformerModelInstance.h
+++ b/Engine/Plugins/Animation/MLDeformer/MLDeformerFramework/Source/MLDeformerFramework/Public/MLDeformerModelInstance.h
@@ -154,7 +154,7 @@ protected:
 	 * @param StartIndex The index where we want to start writing the transformations to.
 	 * @return The new start index into the output buffer, to be used when you write more data after this.
 	 */
-	int64 SetBoneTransforms(float* OutputBuffer, int64 OutputBufferSize, int64 StartIndex);
+	virtual int64 SetBoneTransforms(float* OutputBuffer, int64 OutputBufferSize, int64 StartIndex);
 
 	/**
 	 * Set the animation curve values inside a given output buffer, starting from a given StartIndex.
@@ -163,7 +163,7 @@ protected:
 	 * @param StartIndex The index where we want to start writing the weights to.
 	 * @return The new start index into the output buffer, to be used when you write more data after this.
 	 */
-	int64 SetCurveValues(float* OutputBuffer, int64 OutputBufferSize, int64 StartIndex);
+	virtual int64 SetCurveValues(float* OutputBuffer, int64 OutputBufferSize, int64 StartIndex);
 
 	/**
 	 * Updates the bone transforms array.
diff --git a/Engine/Plugins/Animation/MLDeformer/MLDeformerFramework/Source/MLDeformerFrameworkEditor/Public/MLDeformerGeomCacheSampler.h b/Engine/Plugins/Animation/MLDeformer/MLDeformerFramework/Source/MLDeformerFrameworkEditor/Public/MLDeformerGeomCacheSampler.h
--- a/Engine/Plugins/Animation/MLDeformer/MLDeformerFramework/Source/MLDeformerFrameworkEditor/Public/MLDeformerGeomCacheSampler.h
+++ b/Engine/Plugins/Animation/MLDeformer/MLDeformerFramework/Source/MLDeformerFrameworkEditor/Public/MLDeformerGeomCacheSampler.h
@@ -45,6 +45,9 @@ namespace UE::MLDeformer
 		 */
 		FMLDeformerGetGeomCacheEvent& OnGetGeometryCache()			{ return GetGeometryCacheEvent; }
 
+		UGeometryCacheComponent* GetGeometryCacheComponent() { return GeometryCacheComponent; }
+		UGeometryCacheComponent* GetGeometryCacheComponent() const { return GeometryCacheComponent; }
+
 	protected:
 		/**
 		 * Calculate the vertex deltas between the linear skinned mesh (skeletal mesh) and geometry cache.
diff --git a/Engine/Plugins/Animation/MLDeformer/MLDeformerFramework/Source/MLDeformerFrameworkEditor/Public/MLDeformerSampler.h b/Engine/Plugins/Animation/MLDeformer/MLDeformerFramework/Source/MLDeformerFrameworkEditor/Public/MLDeformerSampler.h
--- a/Engine/Plugins/Animation/MLDeformer/MLDeformerFramework/Source/MLDeformerFrameworkEditor/Public/MLDeformerSampler.h
+++ b/Engine/Plugins/Animation/MLDeformer/MLDeformerFramework/Source/MLDeformerFrameworkEditor/Public/MLDeformerSampler.h
@@ -82,6 +82,7 @@ namespace UE::MLDeformer
 
 		/** Calculate the memory usage of the sampler. */
 		SIZE_T CalcMemUsagePerFrameInBytes() const;
+		UDebugSkelMeshComponent* GetSkeletalMeshComponent() { return SkeletalMeshComponent; }
 
 	protected:
 		/** Create the actors used for sampling. This creates two actors, one for the base skeletal mesh and one for the target mesh. */
diff --git a/Engine/Plugins/Animation/MLDeformer/MLDeformerFramework/Source/MLDeformerFrameworkEditor/Public/MLDeformerTrainingModel.h b/Engine/Plugins/Animation/MLDeformer/MLDeformerFramework/Source/MLDeformerFrameworkEditor/Public/MLDeformerTrainingModel.h
--- a/Engine/Plugins/Animation/MLDeformer/MLDeformerFramework/Source/MLDeformerFrameworkEditor/Public/MLDeformerTrainingModel.h
+++ b/Engine/Plugins/Animation/MLDeformer/MLDeformerFramework/Source/MLDeformerFrameworkEditor/Public/MLDeformerTrainingModel.h
@@ -53,7 +53,7 @@ public:
 	 * This is automatically called by the editor.
 	 * @param InEditorModel The pointer to the editor model that this is a training model for.
 	 */
-	void Init(UE::MLDeformer::FMLDeformerEditorModel* InEditorModel);
+	virtual void Init(UE::MLDeformer::FMLDeformerEditorModel* InEditorModel);
 
 	/** Get the runtime ML Deformer model object. */
 	UFUNCTION(BlueprintPure, Category = "Training Data")
diff --git a/Engine/Plugins/Animation/MLDeformer/NearestNeighborModel/NearestNeighborModel.uplugin b/Engine/Plugins/Animation/MLDeformer/NearestNeighborModel/NearestNeighborModel.uplugin
new file mode 100644
--- /dev/null
+++ b/Engine/Plugins/Animation/MLDeformer/NearestNeighborModel/NearestNeighborModel.uplugin
@@ -0,0 +1,51 @@
+{
+	"FileVersion": 3,
+	"Version": 1,
+	"VersionName": "0.1",
+	"FriendlyName": "ML Deformer Nearest Neighbor Model",
+	"Description": "Nearest Neighbor Model for the ML Deformer Framework",
+	"Category": "Animation",
+	"CreatedBy": "Epic Games, Inc.",
+	"CreatedByURL": "http://epicgames.com",
+	"DocsURL": "https://docs.unrealengine.com/5.0/en-US/using-the-machine-learning-deformer-in-unreal-engine/",
+	"MarketplaceURL": "",
+	"SupportURL": "https://forums.unrealengine.com/",
+	"CanContainContent": true,
+	"IsBetaVersion": false,
+	"IsExperimentalVersion": true,
+	"Installed": false,
+	"Modules": [
+		{
+			"Name": "NearestNeighborModel",
+			"Type": "Runtime",
+			"LoadingPhase": "Default"
+		},
+		{
+			"Name": "NearestNeighborModelEditor",
+			"Type": "Editor",
+			"LoadingPhase": "Default"
+		}
+	],
+	"Plugins": [
+		{
+			"Name": "GeometryCache",
+			"Enabled": true
+		},
+		{
+			"Name": "NeuralNetworkInference",
+			"Enabled": true
+		},
+		{
+			"Name": "DeformerGraph",
+			"Enabled": true
+		},
+		{
+			"Name": "PythonFoundationPackages",
+			"Enabled": true
+		},
+		{
+			"Name": "MLDeformerFramework",
+			"Enabled": true
+		}
+	]
+}
\ No newline at end of file
diff --git a/Engine/Plugins/Animation/MLDeformer/NearestNeighborModel/Source/NearestNeighborModel/NearestNeighborModel.Build.cs b/Engine/Plugins/Animation/MLDeformer/NearestNeighborModel/Source/NearestNeighborModel/NearestNeighborModel.Build.cs
new file mode 100644
--- /dev/null
+++ b/Engine/Plugins/Animation/MLDeformer/NearestNeighborModel/Source/NearestNeighborModel/NearestNeighborModel.Build.cs
@@ -0,0 +1,58 @@
+// Copyright Epic Games, Inc. All Rights Reserved.
+
+using UnrealBuildTool;
+
+public class NearestNeighborModel : ModuleRules
+{
+	public NearestNeighborModel(ReadOnlyTargetRules Target) : base(Target)
+	{
+		PCHUsage = ModuleRules.PCHUsageMode.UseExplicitOrSharedPCHs;
+		
+		PublicIncludePaths.AddRange(
+			new string[] {
+				// ... add public include paths required here ...
+			}
+			);
+				
+		
+		PrivateIncludePaths.AddRange(
+			new string[] {
+				// ... add other private include paths required here ...
+			}
+			);
+			
+		
+		PublicDependencyModuleNames.AddRange(
+			new string[]
+			{
+				"Core"
+				// ... add other public dependencies that you statically link with here ...
+			}
+			);
+			
+		
+		PrivateDependencyModuleNames.AddRange(
+			new string[]
+			{
+				"ComputeFramework",
+				"CoreUObject",
+				"Engine",
+				"GeometryCache",
+				"NeuralNetworkInference",
+				"OptimusCore",
+				"Projects",
+				"RenderCore",
+				"RHI",
+				"MLDeformerFramework",
+			}
+			);
+		
+		
+		DynamicallyLoadedModuleNames.AddRange(
+			new string[]
+			{
+				// ... add any modules that your module loads dynamically here ...
+			}
+			);
+	}
+}
diff --git a/Engine/Plugins/Animation/MLDeformer/NearestNeighborModel/Source/NearestNeighborModel/Private/NearestNeighborModel.cpp b/Engine/Plugins/Animation/MLDeformer/NearestNeighborModel/Source/NearestNeighborModel/Private/NearestNeighborModel.cpp
new file mode 100644
--- /dev/null
+++ b/Engine/Plugins/Animation/MLDeformer/NearestNeighborModel/Source/NearestNeighborModel/Private/NearestNeighborModel.cpp
@@ -0,0 +1,246 @@
+// Copyright Epic Games, Inc. All Rights Reserved.
+#pragma once
+
+#include "NearestNeighborModel.h"
+#include "NearestNeighborModelInputInfo.h"
+#include "NearestNeighborModelInstance.h"
+#include "MLDeformerComponent.h"
+#include "Engine/SkeletalMesh.h"
+#include "GeometryCache.h"
+#include "Misc/FileHelper.h"
+#include "Rendering/SkeletalMeshModel.h"
+#include "UObject/UObjectGlobals.h"
+#include "NeuralNetwork.h"
+
+#define LOCTEXT_NAMESPACE "UNearestNeighborModel"
+
+NEARESTNEIGHBORMODEL_API DEFINE_LOG_CATEGORY(LogNearestNeighborModel)
+
+namespace UE::NearestNeighborModel
+{
+	class NEARESTNEIGHBORMODEL_API FNearestNeighborModelModule
+		: public IModuleInterface
+	{
+	};
+}
+IMPLEMENT_MODULE(UE::NearestNeighborModel::FNearestNeighborModelModule, NearestNeighborModel)
+
+
+UNearestNeighborModel::UNearestNeighborModel(const FObjectInitializer& ObjectInitializer)
+	: Super(ObjectInitializer)
+{
+#if WITH_EDITORONLY_DATA
+	VizSettings = ObjectInitializer.CreateEditorOnlyDefaultSubobject<UNearestNeighborModelVizSettings>(this, TEXT("VizSettings"));
+#endif
+
+	ExternalMorphSetID = 100;
+}
+
+UMLDeformerInputInfo* UNearestNeighborModel::CreateInputInfo()
+{
+	UNearestNeighborModelInputInfo* NearestNeighborModelInputInfo = NewObject<UNearestNeighborModelInputInfo>(this);
+#if WITH_EDITORONLY_DATA
+	NearestNeighborModelInputInfo->InitRefBoneRotations(SkeletalMesh);
+#endif
+	return NearestNeighborModelInputInfo;
+}
+
+UMLDeformerModelInstance* UNearestNeighborModel::CreateModelInstance(UMLDeformerComponent* Component)
+{
+	return NewObject<UNearestNeighborModelInstance>(Component);
+}
+
+void UNearestNeighborModel::PostLoad()
+{
+	Super::PostLoad();
+
+	InitInputInfo();
+	InitPreviousWeights();
+
+#if WITH_EDITORONLY_DATA
+	NearestNeighborIds.SetNumZeroed(GetNumParts());
+	UpdateNetworkInputDim();
+	UpdateNetworkOutputDim();
+#endif
+}
+
+
+TArray<uint32> ReadTxt(const FString &Path)
+{
+	TArray<FString> StrArr;
+	FFileHelper::LoadFileToStringArray(StrArr, *Path);
+	TArray<uint32> UIntArr; UIntArr.SetNum(StrArr.Num());
+	for (int32 i = 0; i < UIntArr.Num(); i++)
+	{
+		UIntArr[i] = FCString::Atoi(*StrArr[i]);
+	}
+
+	return UIntArr;
+}
+
+void UNearestNeighborModel::ClipInputs(float* InputPtr, int NumInputs)
+{
+	if(InputsMin.Num() == NumInputs && InputsMax.Num() == NumInputs)
+	{
+		for(int32 i = 0; i < NumInputs; i++)
+		{
+			if (InputPtr[i] > InputsMax[i])
+			{
+				InputPtr[i] = InputsMax[i];
+			}
+			if (InputPtr[i] < InputsMin[i])
+			{
+				InputPtr[i] = InputsMin[i];
+			}
+		}
+	}	
+}
+
+#if WITH_EDITORONLY_DATA
+void UNearestNeighborModel::UpdatePCACoeffNums()
+{
+	uint32 PCACoeffStart = 0;
+	for(int32 PartId = 0; PartId < GetNumParts(); PartId++)
+	{
+		ClothPartData[PartId].PCACoeffStart = PCACoeffStart;
+		PCACoeffStart+=ClothPartData[PartId].PCACoeffNum;
+	}
+}
+
+void UNearestNeighborModel::UpdateNetworkInputDim()
+{
+	InputDim = 3 * BoneIncludeList.Num();
+}
+
+void UNearestNeighborModel::UpdateNetworkOutputDim()
+{
+	OutputDim = 0;
+	for (int32 i = 0; i < GetNumParts(); i++)
+	{
+		OutputDim += ClothPartData[i].PCACoeffNum; 
+	}
+}
+
+void UNearestNeighborModel::UpdateClothPartData()
+{
+	if(ClothPartData.Num() == 0)
+	{
+		UE_LOG(LogNearestNeighborModel, Error, TEXT("There should be at least 1 cloth part"));
+	}
+
+	ClothPartData.SetNum(ClothPartEditorData.Num());
+	for(int32 PartId = 0; PartId < GetNumParts(); PartId++)
+	{
+		if(FPaths::FileExists(ClothPartEditorData[PartId].VertexMapPath))
+		{
+			ClothPartData[PartId].PCACoeffNum = ClothPartEditorData[PartId].PCACoeffNum;
+			ClothPartData[PartId].VertexMap = ReadTxt(ClothPartEditorData[PartId].VertexMapPath);
+			ClothPartData[PartId].NumVertices = ClothPartData[PartId].VertexMap.Num();
+
+			// Initialize PCA data with 0 to prevent NearestNeighborModelInstance::Tick from crashing
+			if (!CheckPCAData(PartId))
+			{
+				ClothPartData[PartId].VertexMean.SetNumZeroed(ClothPartData[PartId].NumVertices * 3);
+				ClothPartData[PartId].PCABasis.SetNumZeroed(ClothPartData[PartId].NumVertices * 3 * ClothPartData[PartId].PCACoeffNum);
+
+				// NumNeighbors cannot be 0 because of the NearestNeighborModel.usf shader.
+				// Init default neighbor data.
+				ClothPartData[PartId].NeighborCoeffs.SetNumZeroed(ClothPartData[PartId].PCACoeffNum);
+				ClothPartData[PartId].NeighborOffsets.SetNumZeroed(ClothPartData[PartId].NumVertices * 3);
+				ClothPartData[PartId].NumNeighbors = 1;
+			}
+		}
+		else
+		{
+			UE_LOG(LogNearestNeighborModel, Error, TEXT("%s does not exist"), *ClothPartEditorData[PartId].VertexMapPath);
+			return;
+		}
+	}
+	UpdatePCACoeffNums();
+
+	NearestNeighborIds.SetNumZeroed(GetNumParts());
+	NearestNeighborData.SetNumZeroed(GetNumParts());
+	UpdateNetworkInputDim();
+	UpdateNetworkOutputDim();
+	bClothPartDataValid = true;
+}
+
+TObjectPtr<UAnimSequence> UNearestNeighborModel::GetNearestNeighborSkeletons(int32 PartId)
+{ 
+	if (PartId < NearestNeighborData.Num())
+	{
+		return NearestNeighborData[PartId].Skeletons; 
+	}
+	return nullptr;
+}
+
+const TObjectPtr<UAnimSequence> UNearestNeighborModel::GetNearestNeighborSkeletons(int32 PartId) const
+{ 
+	if (PartId < NearestNeighborData.Num())
+	{
+		return NearestNeighborData[PartId].Skeletons; 
+	}
+	return nullptr;
+}
+
+TObjectPtr<UGeometryCache> UNearestNeighborModel::GetNearestNeighborCache(int32 PartId)
+{ 
+	if (PartId < NearestNeighborData.Num())
+	{
+		return NearestNeighborData[PartId].Cache;
+	}
+	return nullptr;
+}
+
+const TObjectPtr<UGeometryCache> UNearestNeighborModel::GetNearestNeighborCache(int32 PartId) const
+{ 
+	if (PartId < NearestNeighborData.Num())
+	{
+		return NearestNeighborData[PartId].Cache;
+	}
+	return nullptr;
+}
+
+int32 UNearestNeighborModel::GetNumNeighborsFromGeometryCache(int32 PartId) const
+{
+	const TObjectPtr<UGeometryCache> Cache = GetNearestNeighborCache(PartId);
+	if (!Cache.IsNull())
+	{
+		const int32 StartFrame = Cache->GetStartFrame();
+		const int32 EndFrame = Cache->GetEndFrame();
+		return EndFrame - StartFrame + 1;
+	}
+	else
+	{
+		return 0;
+	}
+}
+#endif
+
+void UNearestNeighborModel::InitInputInfo()
+{
+	UNearestNeighborModelInputInfo* NearestNeighborModelInputInfo = static_cast<UNearestNeighborModelInputInfo*>(GetInputInfo());
+	check(NearestNeighborModelInputInfo != nullptr);
+#if WITH_EDITORONLY_DATA
+	NearestNeighborModelInputInfo->InitRefBoneRotations(SkeletalMesh);
+#endif
+}
+
+bool UNearestNeighborModel::CheckPCAData(int32 PartId) const
+{
+	const FClothPartData& Data = ClothPartData[PartId];
+	return Data.VertexMap.Num() > 0 && Data.PCABasis.Num() == Data.VertexMap.Num() * 3 * Data.PCACoeffNum;
+}
+
+
+void UNearestNeighborModel::InitPreviousWeights()
+{
+	int32 NumWeights = 1;
+	for (int32 PartId = 0; PartId < GetNumParts(); PartId++)
+	{
+		NumWeights += GetPCACoeffNum(PartId);
+		NumWeights += GetNumNeighbors(PartId);
+	}
+	PreviousWeights.SetNumZeroed(NumWeights);
+}
+#undef LOCTEXT_NAMESPACE
diff --git a/Engine/Plugins/Animation/MLDeformer/NearestNeighborModel/Source/NearestNeighborModel/Private/NearestNeighborModelInputInfo.cpp b/Engine/Plugins/Animation/MLDeformer/NearestNeighborModel/Source/NearestNeighborModel/Private/NearestNeighborModelInputInfo.cpp
new file mode 100644
--- /dev/null
+++ b/Engine/Plugins/Animation/MLDeformer/NearestNeighborModel/Source/NearestNeighborModel/Private/NearestNeighborModelInputInfo.cpp
@@ -0,0 +1,81 @@
+// Copyright Epic Games, Inc. All Rights Reserved.
+
+#include "NearestNeighborModelInputInfo.h"
+#include "Components/SkeletalMeshComponent.h"
+
+
+void UNearestNeighborModelInputInfo::ComputeNetworkInput(const TArray<FQuat>& BoneRotations, float* OutputBuffer, int64 StartIndex) const
+{
+    if (OutputBuffer == nullptr)
+    {
+        return;
+    }
+
+    const int32 NumBones = GetNumBones();
+    check(BoneRotations.Num() == NumBones && RefBoneRotations.Num() == NumBones);
+
+    int64 Index = StartIndex;
+    for(int32 BoneIndex = 0; BoneIndex < NumBones; BoneIndex++)
+    {
+        FQuat RestR = RefBoneRotations[BoneIndex];
+        FQuat LocalR = BoneRotations[BoneIndex];
+
+        const FQuat R =  RestR.Inverse() * LocalR;
+        const float Scalar = R.W >= 0 ? 2.f : -2.f;
+
+        OutputBuffer[Index++] = R.X * Scalar;
+        OutputBuffer[Index++] = R.Y * Scalar;
+        OutputBuffer[Index++] = R.Z * Scalar;
+    }
+}
+
+void UNearestNeighborModelInputInfo::ExtractBoneRotations(USkeletalMeshComponent* SkelMeshComponent, TArray<float>& OutRotations) const
+{
+    const int32 NumBones = GetNumBones();
+    const int32 NumFloats = NumBones * 3;
+    OutRotations.Reset(NumFloats);
+    OutRotations.AddUninitialized(NumFloats);
+
+    const TArray<FTransform>& BoneTransforms = SkelMeshComponent->GetBoneSpaceTransforms();
+
+    // Extract BoneRotations from BoneTransforms
+    TArray<FQuat> BoneRotations;
+    BoneRotations.AddZeroed(NumBones);
+    for (int32 Index = 0; Index < NumBones; ++Index)
+    {
+        const FName BoneName = GetBoneName(Index);
+        const int32 SkelMeshBoneIndex = SkelMeshComponent->GetBoneIndex(BoneName);
+        FVector RotationVector = FVector::Zero();
+        if(SkelMeshBoneIndex != INDEX_NONE)
+        {
+            BoneRotations[Index] = BoneTransforms[SkelMeshBoneIndex].GetRotation();
+        }
+    }
+
+    ComputeNetworkInput(BoneRotations, OutRotations.GetData());
+}
+
+int32 UNearestNeighborModelInputInfo::CalcNumNeuralNetInputs() const
+{
+    return 
+        BoneNameStrings.Num() * 3 + // Three floats per bone.
+        CurveNameStrings.Num();     // One float per curve.
+}
+
+void UNearestNeighborModelInputInfo::InitRefBoneRotations(USkeletalMesh* SkelMesh)
+{
+    RefBoneRotations.Reset();
+    if (SkelMesh)
+    {
+        const FReferenceSkeleton& RefSkel = SkelMesh->GetRefSkeleton();
+        const TArray<FTransform>& RefBonePose = RefSkel.GetRefBonePose();
+        const int32 NumBones = GetNumBones();
+        RefBoneRotations.SetNumUninitialized(NumBones);
+        for(int32 Index = 0; Index < NumBones; Index++)
+        {
+            const FName BoneName = GetBoneName(Index);
+            const int32 BoneIndex = RefSkel.FindBoneIndex(BoneName);
+            RefBoneRotations[Index] = RefBonePose[BoneIndex].GetRotation();
+        }
+    }
+}
\ No newline at end of file
diff --git a/Engine/Plugins/Animation/MLDeformer/NearestNeighborModel/Source/NearestNeighborModel/Private/NearestNeighborModelInstance.cpp b/Engine/Plugins/Animation/MLDeformer/NearestNeighborModel/Source/NearestNeighborModel/Private/NearestNeighborModelInstance.cpp
new file mode 100644
--- /dev/null
+++ b/Engine/Plugins/Animation/MLDeformer/NearestNeighborModel/Source/NearestNeighborModel/Private/NearestNeighborModelInstance.cpp
@@ -0,0 +1,208 @@
+// Copyright Epic Games, Inc. All Rights Reserved.
+
+#include "NearestNeighborModelInstance.h"
+#include "NearestNeighborModel.h"
+#include "NearestNeighborModelInputInfo.h"
+#include "NeuralNetwork.h"
+#include "Components/SkeletalMeshComponent.h"
+
+using namespace UE::MLDeformer;
+using namespace UE::NearestNeighborModel;
+
+int64 UNearestNeighborModelInstance::SetBoneTransforms(float* OutputBuffer, int64 OutputBufferSize, int64 StartIndex)
+{
+	// Get the transforms for the bones we used during training.
+	// These are in the space relative to their parent.
+	UpdateBoneTransforms();
+
+    const UNearestNeighborModelInputInfo* InputInfo = static_cast<UNearestNeighborModelInputInfo*>(Model->GetInputInfo());
+    const int32 NumBones = InputInfo->GetNumBones();
+
+    // Extract BoneRotations from TrainingBoneTransforms (@todo: optimize)
+    TArray<FQuat> BoneRotations; BoneRotations.AddZeroed(NumBones);
+    int64 EndIndex = StartIndex + NumBones * 3;
+	check(EndIndex <= OutputBufferSize); // Make sure we don't write past the OutputBuffer.
+    
+	for (int32 BoneIndex = 0; BoneIndex < NumBones; ++BoneIndex)
+	{
+        BoneRotations[BoneIndex] = TrainingBoneTransforms[BoneIndex].GetRotation();
+	}
+
+    // Write the transforms into the output buffer.
+    InputInfo->ComputeNetworkInput(BoneRotations, OutputBuffer, StartIndex);
+
+	return EndIndex;
+}
+
+BEGIN_SHADER_PARAMETER_STRUCT(FCopyBufferParameters, )
+    RDG_BUFFER_ACCESS(Source,  ERHIAccess::CopySrc)
+END_SHADER_PARAMETER_STRUCT()
+
+template<class T, class TRDGBuffer>
+void AddCopyToCPUPass(FRDGBuilder& GraphBuilder, TRDGBuffer RDGBuffer, TArray<T>& Array)
+{
+    FCopyBufferParameters* Parameters = GraphBuilder.AllocParameters<FCopyBufferParameters>();
+    Parameters->Source = RDGBuffer;
+	GraphBuilder.AddPass(
+		RDG_EVENT_NAME("CopyToCPU"),
+        Parameters,
+        ERDGPassFlags::Readback,
+		[RDGBuffer, &Array](FRHICommandListImmediate& CmdList)
+		{
+			const int32 VolumeInBytes = Array.Num() * sizeof(T);
+			FRHIBuffer* RHIBuffer = RDGBuffer->GetRHI();
+			const void* const BufferData = CmdList.LockBuffer(RHIBuffer, 0, VolumeInBytes, RLM_ReadOnly);
+			FMemory::Memcpy((void*)Array.GetData(), BufferData, VolumeInBytes);
+			CmdList.UnlockBuffer(RHIBuffer);
+		});
+}
+
+template<class T, class TRDGBuffer>
+void AddCopyToCPUPass(FRDGBuilder& GraphBuilder, TRDGBuffer RDGBuffer, T* Ptr, int NumElements)
+{
+    FCopyBufferParameters* Parameters = GraphBuilder.AllocParameters<FCopyBufferParameters>();
+    Parameters->Source = RDGBuffer;
+    GraphBuilder.AddPass(
+        RDG_EVENT_NAME("CopyToCPU"),
+        Parameters,
+        ERDGPassFlags::Readback,
+        [RDGBuffer, Ptr, NumElements](FRHICommandListImmediate& CmdList)
+        {
+            const int32 VolumeInBytes = NumElements * sizeof(T);
+            FRHIBuffer* RHIBuffer = RDGBuffer->GetRHI();
+            const void* const BufferData = CmdList.LockBuffer(RHIBuffer, 0, VolumeInBytes, RLM_ReadOnly);
+            FMemory::Memcpy((void*)Ptr, BufferData, VolumeInBytes);
+            CmdList.UnlockBuffer(RHIBuffer);
+        });
+}
+
+bool UNearestNeighborModelInstance::SetupNeuralNetworkForFrame()
+{
+    bool bSuccess = UMLDeformerModelInstance::SetupNeuralNetworkForFrame();
+
+    // Clip network inputs based on training set min and max
+    UNeuralNetwork* NeuralNetwork = Model->GetNeuralNetwork();
+    if (NeuralNetwork)
+    {
+    	UNearestNeighborModel *NearestNeighborModel = static_cast<UNearestNeighborModel*>(Model.Get());
+    	float* InputDataPointer = static_cast<float*>(NeuralNetwork->GetInputDataPointerMutableForContext(NeuralNetworkInferenceHandle));
+    	const int64 NumNeuralNetInputs = NeuralNetwork->GetInputTensorForContext(NeuralNetworkInferenceHandle).Num();
+    	NearestNeighborModel->ClipInputs(InputDataPointer, NumNeuralNetInputs);
+
+    	return bSuccess;
+    }
+    else
+    {
+    	return false;
+    }
+}
+
+void UNearestNeighborModelInstance::RunNeuralNetwork(float ModelWeight)
+{
+	Super::RunNeuralNetwork(ModelWeight);
+	RunNearestNeighborModel(ModelWeight);
+}
+
+void UNearestNeighborModelInstance::RunNearestNeighborModel(float ModelWeight)
+{
+	UNearestNeighborModel* NearestNeighborModel = Cast<UNearestNeighborModel>(Model);
+	if (NearestNeighborModel == nullptr)
+	{
+		return;
+	}
+
+	const int32 ExternalMorphSetID = NearestNeighborModel->GetExternalMorphSetID();
+	// If this check fails please set this member to some value larger than 0 in your model's constructor.
+	checkf(ExternalMorphSetID != -1, TEXT("Please set the ExternalMorphSetID member value to a unique value for your model type."));
+
+	// Grab the weight data for this morph set.
+	// This could potentially fail if we are applying this deformer to the wrong skeletal mesh component.
+	const int LOD = 0;	// For now we only support LOD 0, as we can't setup an ML Deformer per LOD yet.
+	FExternalMorphSetWeights* WeightData = SkeletalMeshComponent->GetExternalMorphWeights(LOD).MorphSets.Find(ExternalMorphSetID);
+	if (WeightData == nullptr)
+	{
+		return;
+	}
+
+	if (ModelWeight > 0.0f)
+	{
+		const UNeuralNetwork* NeuralNetwork = Model->GetNeuralNetwork();
+		const FNeuralTensor& OutputTensor = NeuralNetwork->GetOutputTensorForContext(NeuralNetworkInferenceHandle);
+		const int32 NumNetworkWeights = OutputTensor.Num();
+		const int32 NumMorphTargets = WeightData->Weights.Num();;
+		if (NumMorphTargets >= NumNetworkWeights + 1)
+		{
+			WeightData->Weights[0] = 1.0f * ModelWeight;
+
+			// Update all generated morph target weights with the values calculated by our neural network.
+			for (int32 MorphIndex = 0; MorphIndex < NumNetworkWeights; ++MorphIndex)
+			{
+				const float W = OutputTensor.At<float>(MorphIndex) * ModelWeight;
+				UpdateWeight(WeightData->Weights, MorphIndex + 1, W);
+			}
+
+			int32 NeighborOffset = NumNetworkWeights + 1;
+			for (int32 PartId = 0; PartId < NearestNeighborModel->GetNumParts(); PartId++)
+			{
+				const int32 NearestNeighborId = FindNearestNeighbor(OutputTensor, PartId);
+#if WITH_EDITORONLY_DATA
+				NearestNeighborModel->SetNearestNeighborId(PartId, NearestNeighborId);
+#endif
+
+				const int32 NumNeighbors = NearestNeighborModel->GetNumNeighbors(PartId);
+				for (int32 NeighborId = 0; NeighborId < NumNeighbors; NeighborId++)
+				{
+					 const float W = NeighborId == NearestNeighborId ? ModelWeight : 0;
+					 UpdateWeight(WeightData->Weights, NeighborOffset + NeighborId, W);
+				}
+				NeighborOffset += NumNeighbors;
+			}
+		}
+	}
+	else
+	{
+		for (float& Weight : WeightData->Weights)
+		{
+			Weight = 0.0f;
+		}
+	}
+}
+
+void UNearestNeighborModelInstance::UpdateWeight(TArray<float>& MorphWeights, const int32 Index, const float W)
+{
+	UNearestNeighborModel* NearestNeighborModel = static_cast<UNearestNeighborModel*>(Model);
+	const float DecayFactor = NearestNeighborModel->GetDecayFactor();
+
+	const float PreviousW = NearestNeighborModel->PreviousWeights[Index];
+	const float NewW = (1 - DecayFactor) * W + DecayFactor * PreviousW;
+	MorphWeights[Index] = NewW;
+	NearestNeighborModel->PreviousWeights[Index] = NewW; 
+}
+
+int32 UNearestNeighborModelInstance::FindNearestNeighbor(const FNeuralTensor& PCACoeffTensor, int32 PartId)
+{
+	const UNearestNeighborModel* NearestNeighborModel = Cast<UNearestNeighborModel>(Model);
+	const TArray<float> &NeighborCoeffs = NearestNeighborModel->NeighborCoeffs(PartId);
+	float MinD2 = std::numeric_limits<float>::max();
+	int32 MinId = -1;
+	for (int32 NeighborId = 0; NeighborId < NearestNeighborModel->GetNumNeighbors(PartId); NeighborId++)
+	{
+		float D2 = 0;
+		const int32 NumCoeffs = NearestNeighborModel->GetPCACoeffNum(PartId);
+		const int32 CoeffStart = NearestNeighborModel->GetPCACoeffStart(PartId);
+		const int32 NumNeighbors = NearestNeighborModel->GetNumNeighbors(PartId);
+		for (int32 CoeffId = 0; CoeffId < NumCoeffs; CoeffId++)
+		{
+			float Coeff = PCACoeffTensor.At<float>(CoeffStart + CoeffId);
+			float NeighborCoeff = NeighborCoeffs[CoeffId * NumNeighbors + NeighborId];
+			float D = Coeff - NeighborCoeff;
+			D2 += D * D;
+		}
+		if (D2 < MinD2)
+		{
+			MinD2 = D2;
+			MinId = NeighborId;
+		}
+	}
+	return MinId;
+}
\ No newline at end of file
diff --git a/Engine/Plugins/Animation/MLDeformer/NearestNeighborModel/Source/NearestNeighborModel/Private/Tests/NearestNeighborModelTest.cpp b/Engine/Plugins/Animation/MLDeformer/NearestNeighborModel/Source/NearestNeighborModel/Private/Tests/NearestNeighborModelTest.cpp
new file mode 100644
--- /dev/null
+++ b/Engine/Plugins/Animation/MLDeformer/NearestNeighborModel/Source/NearestNeighborModel/Private/Tests/NearestNeighborModelTest.cpp
@@ -0,0 +1,122 @@
+// Copyright Epic Games, Inc. All Rights Reserved.
+
+#include "CoreMinimal.h"
+#include "Misc/AutomationTest.h"
+
+#if WITH_DEV_AUTOMATION_TESTS
+
+#include "NearestNeighborModel.h"
+#include "NearestNeighborModelVizSettings.h"
+#include "MLDeformerAsset.h"
+#include "MLDeformerComponent.h"
+#include "MLDeformerModelInstance.h"
+#include "MLDeformerGeomCacheHelpers.h"
+#include "UObject/GCObjectScopeGuard.h"
+#include "PreviewScene.h"
+#include "Components/SkeletalMeshComponent.h"
+#include "NeuralNetwork.h"
+#include "GeometryCache.h"
+#include "Engine/SkeletalMesh.h"
+#include "Rendering/SkeletalMeshModel.h"
+#include "Rendering/SkeletalMeshLODModel.h"
+
+namespace UE::NearestNeighborModelTests
+{
+	IMPLEMENT_SIMPLE_AUTOMATION_TEST(FNearestNeighborModelMainTest, "MLDeformer.NearestNeighborModel.MainTest", EAutomationTestFlags::ApplicationContextMask | EAutomationTestFlags::EngineFilter);
+	bool FNearestNeighborModelMainTest::RunTest(const FString& Parameters)
+	{
+		using namespace UE::MLDeformer;
+		using namespace UE::NearestNeighborModel;
+
+		UMLDeformerAsset* DeformerAsset = LoadObject<UMLDeformerAsset>(nullptr, TEXT("MLDeformerAsset'/NearestNeighborModel/Tests/Biceps/biceps_mld.biceps_mld'"));
+		UTEST_NOT_NULL(TEXT("MLDeformerAsset load"), DeformerAsset);
+		FGCObjectScopeGuard Guard(DeformerAsset);
+
+		UNearestNeighborModel* NearestNeighborModel = Cast<UNearestNeighborModel>(DeformerAsset->GetModel());
+		UTEST_NOT_NULL(TEXT("NearestNeighborModel load check"), NearestNeighborModel);
+
+		UNeuralNetwork* NeuralNet = NearestNeighborModel->GetNeuralNetwork();
+		UTEST_NOT_NULL(TEXT("Neuralnet check"), NeuralNet);
+
+		UMLDeformerInputInfo* InputInfo = NearestNeighborModel->GetInputInfo();
+		UTEST_NOT_NULL(TEXT("InputInfo check"), InputInfo);
+		const int64 NumInputInfoInputs = InputInfo->CalcNumNeuralNetInputs();
+		const int64 NumNetworkInputs = NeuralNet->GetInputTensor().Num();
+		UTEST_EQUAL(TEXT("InputInfo input count check"), NumInputInfoInputs, NumNetworkInputs);
+
+		const int32 NumInputBones = InputInfo->GetNumBones();
+		UTEST_EQUAL(TEXT("InputInfo bone count check"), NumInputBones, 1);
+		UTEST_EQUAL(TEXT("InputInfo bone name string check"), InputInfo->GetBoneNameStrings().Num(), 1);
+		UTEST_EQUAL(TEXT("InputInfo bone name string content check"), InputInfo->GetBoneNameString(0), FString("lowerarm_l"));
+		UTEST_EQUAL(TEXT("InputInfo bone names check"), InputInfo->GetBoneNames().Num(), 1);
+		UTEST_EQUAL(TEXT("InputInfo bone names content check"), InputInfo->GetBoneName(0), FName("lowerarm_l"));
+		UTEST_EQUAL(TEXT("InputInfo vertex count check"), InputInfo->GetNumBaseMeshVertices(), InputInfo->GetNumTargetMeshVertices());
+
+		const FVertexMapBuffer& VertexMapBuffer = NearestNeighborModel->GetVertexMapBuffer();
+		UTEST_TRUE(TEXT("VertexMapBuffer check"), VertexMapBuffer.IsInitialized());
+
+#if WITH_EDITORONLY_DATA
+		USkeletalMesh* SkeletalMesh = NearestNeighborModel->GetSkeletalMesh();
+		UTEST_NOT_NULL(TEXT("SkeletalMesh check"), SkeletalMesh);
+
+		UTEST_TRUE(TEXT("Model vs Network compatible check"), NearestNeighborModel->GetInputInfo()->IsCompatible(SkeletalMesh));
+
+		UGeometryCache* GeomCache = NearestNeighborModel->GetGeometryCache();
+		UTEST_NOT_NULL(TEXT("GeomCache check"), GeomCache);
+
+		const FSkeletalMeshModel* ImportedModel = SkeletalMesh->GetImportedModel();
+		UTEST_NOT_NULL(TEXT("ImportedModel check"), ImportedModel);
+
+		const TArray<int32>& VertexMap = NearestNeighborModel->GetVertexMap();
+		UTEST_EQUAL(TEXT("VertexMap size check"), VertexMap.Num(), ImportedModel->LODModels[0].NumVertices);
+
+		const int32 NumBaseMeshVerts = UMLDeformerModel::ExtractNumImportedSkinnedVertices(SkeletalMesh);
+		const int32 NumGeomCacheVerts = ExtractNumImportedGeomCacheVertices(GeomCache);
+		UTEST_EQUAL(TEXT("VertexCount check"), NumBaseMeshVerts, NumGeomCacheVerts);
+		UTEST_EQUAL(TEXT("Model SkelMesh VertexCount check"), NearestNeighborModel->GetNumBaseMeshVerts(), NumBaseMeshVerts);
+		UTEST_EQUAL(TEXT("Model TargetMesh VertexCount check"), NearestNeighborModel->GetNumTargetMeshVerts(), NumGeomCacheVerts);
+
+		const int64 NumNetworkOutputs = NeuralNet->GetOutputTensor().Num();
+		const int64 ExpectedNetworkOutput = NumGeomCacheVerts * 3;	// 3 floats per vertex.
+		UTEST_EQUAL(TEXT("NeuralNet output size check"), NumNetworkOutputs, ExpectedNetworkOutput);
+		// UTEST_EQUAL(TEXT("Hidden Layers check"), NearestNeighborModel->GetNumHiddenLayers(), 3);
+
+		UNearestNeighborModelVizSettings* VizSettings = Cast<UNearestNeighborModelVizSettings>(NearestNeighborModel->GetVizSettings());
+		UTEST_NOT_NULL(TEXT("VizSettings check"), VizSettings);
+
+		// Create a preview scene and spawn an actor.
+		TUniquePtr<FPreviewScene> PreviewScene = MakeUnique<FPreviewScene>(FPreviewScene::ConstructionValues());
+		UWorld* World = PreviewScene->GetWorld();
+		UTEST_NOT_NULL(TEXT("World check"), World);
+		FActorSpawnParameters SpawnParams;
+		SpawnParams.Name = MakeUniqueObjectName(World, AActor::StaticClass(), FName(TEXT("ML Deformer Test Actor")));	
+		AActor* Actor = World->SpawnActor<AActor>(SpawnParams);
+		Actor->SetFlags(RF_Transient);
+		UTEST_NOT_NULL(TEXT("Actor check"), Actor);
+
+		// Add a skeletal mesh component.
+		USkeletalMeshComponent* SkelMeshComponent = NewObject<USkeletalMeshComponent>(Actor);
+		SkelMeshComponent->SetSkeletalMesh(SkeletalMesh);
+		Actor->SetRootComponent(SkelMeshComponent);
+		SkelMeshComponent->RegisterComponent();
+		
+		// Add an ML Deformer component and activate it.
+		UMLDeformerComponent* MLDeformerComponent = NewObject<UMLDeformerComponent>(Actor);
+		MLDeformerComponent->SetDeformerAsset(DeformerAsset);
+		MLDeformerComponent->RegisterComponent();
+		MLDeformerComponent->Activate();
+		UTEST_EQUAL(TEXT("MLDeformerComponent weight check"), MLDeformerComponent->GetWeight(), 1.0f);
+		UTEST_TRUE(TEXT("MLDeformerComponent SkelMeshComponent check"), MLDeformerComponent->GetSkeletalMeshComponent() == SkelMeshComponent);
+
+		// Create and init the model instance.
+		UMLDeformerModelInstance* ModelInstance = MLDeformerComponent->GetModelInstance();
+		UTEST_NOT_NULL(TEXT("CreateModelInstance check"), ModelInstance);
+		UTEST_TRUE(TEXT("ModelInstance GetModel check"), ModelInstance->GetModel() == NearestNeighborModel);
+		UTEST_TRUE(TEXT("ModelInstance compatible check"), ModelInstance->IsCompatible());
+#endif // WITH_EDITORONLY_DATA
+
+		return true;
+	}
+}	// namespace UE::NearestNeighborModelTests
+
+#endif	// #if WITH_DEV_AUTOMATION_TESTS
diff --git a/Engine/Plugins/Animation/MLDeformer/NearestNeighborModel/Source/NearestNeighborModel/Public/NearestNeighborModel.h b/Engine/Plugins/Animation/MLDeformer/NearestNeighborModel/Source/NearestNeighborModel/Public/NearestNeighborModel.h
new file mode 100644
--- /dev/null
+++ b/Engine/Plugins/Animation/MLDeformer/NearestNeighborModel/Source/NearestNeighborModel/Public/NearestNeighborModel.h
@@ -0,0 +1,300 @@
+// Copyright Epic Games, Inc. All Rights Reserved.
+#pragma once
+
+#include "CoreMinimal.h"
+#include "MLDeformerModel.h"
+#include "MLDeformerVizSettings.h"
+#include "MLDeformerInputInfo.h"
+#include "MLDeformerMorphModel.h"
+#include "NearestNeighborModelVizSettings.h"
+#include "UObject/Object.h"
+#include "NearestNeighborModel.generated.h"
+
+class USkeletalMesh;
+class UGeometryCache;
+class UAnimSequence;
+class UNeuralNetwork;
+class UMLDeformerAsset;
+class USkeleton;
+class IPropertyHandle;
+
+NEARESTNEIGHBORMODEL_API DECLARE_LOG_CATEGORY_EXTERN(LogNearestNeighborModel, Log, All);
+
+namespace UE::NearestNeighborModel
+{
+	template<class T>
+	class FArrayBufferWithSRV : public FVertexBufferWithSRV
+	{
+	public:
+		void Init(TArray<T>* InArray, const FString& InDebugName);
+
+	private:
+		void InitRHI() override;
+
+		TArray<T> *Array = nullptr;
+		FString DebugName;
+	};
+
+	template<class T>
+	class FVertexBufferWithUAV : public FVertexBufferWithSRV
+	{
+	public:
+		void Init(int32 InNumElements, const FString& InDebugName);
+
+	private:
+		void InitRHI() override;
+
+		int32 NumElements = 0;
+		FString DebugName;
+	};
+};
+
+USTRUCT(BlueprintType, Blueprintable)
+struct FSkeletonCachePair
+{
+	GENERATED_BODY()
+
+	UPROPERTY(EditAnywhere, BlueprintReadWrite, Category = "Nearest Neighbors")
+	TObjectPtr<UAnimSequence> Skeletons = nullptr;
+
+	UPROPERTY(EditAnywhere, BlueprintReadWrite, Category = "Nearest Neighbors")
+	TObjectPtr<UGeometryCache> Cache = nullptr;
+};
+
+USTRUCT(BlueprintType, Blueprintable)
+struct FClothPartEditorData
+{
+	GENERATED_BODY()
+
+	UPROPERTY(EditAnywhere, BlueprintReadWrite, meta = (ClampMin = "1", ClampMax = "128"), Category = "Nearest Neighbors")
+	int32 PCACoeffNum = 128;
+
+	UPROPERTY(EditAnywhere, BlueprintReadWrite, Category = "Nearest Neighbors")
+	FString VertexMapPath;
+};
+
+USTRUCT(BlueprintType, Blueprintable)
+struct FClothPartData
+{
+	GENERATED_BODY()
+
+	/** Number of PCA coefficients for this cloth part. */
+	UPROPERTY()
+	int32 PCACoeffNum = 128;
+
+	/** The start index of PCA coeffcients of this cloth part */
+	UPROPERTY()
+	uint32 PCACoeffStart = 0;
+
+	/** Number of vertices in this cloth part */
+	UPROPERTY()
+	uint32 NumVertices = 0;
+
+	/** Number of neighbors used for this cloth part */
+	UPROPERTY()
+	uint32 NumNeighbors = 0;
+
+	/** Vertex indices for this cloth part */
+	UPROPERTY()
+	TArray<uint32> VertexMap;
+
+	/** PCA basis for this cloth part. This is a flattened array of size (PCACoeffNum, NumVertices * 3)  */
+	UPROPERTY()
+	TArray<float> PCABasis;
+
+	/** Vertex mean for PCA computation. This has the size of NumVertices * 3 */
+	UPROPERTY()
+	TArray<float> VertexMean;
+
+	/** PCA coefficients of the nearest neighbors. This is a flattened array of size (PCACoeffNum, NumNeighbors) */
+	UPROPERTY()
+	TArray<float> NeighborCoeffs;
+
+	/** The remaining offsets of the nearest neighbor shapes (after reducing PCA offsets). This is a flattened array of size (NumNeighbors, PCACoeffNum) */
+	UPROPERTY()
+	TArray<float> NeighborOffsets;
+};
+
+/**
+ * The nearest neighbor model.
+ * This model contains the PCA basis of the cloth vertex deltas and a small set of cloth for nearest neighbor search. 
+ * Given a new pose, the pre-trained neural network first predicts the PCA coefficients of the vertex deltas. 
+ * Then this model uses the predicted PCA coeffcients to find a nearest neighbor in the small cloth set.
+ * The total vertex delta is computed by
+ * 		vertex_delta = mean_delta + pca_basis * pca_coeff + nearest_neighbor_delta
+ * To prevent popping, a time filtering is applied on predicted vertex deltas. The vertex delta at time t is computed by  
+ * 		vertex_delta(t) = decay_factor * vertex_delta(t-1) + (1 - decay_factor) * vertex_delta 
+ * The cloth can be separated into several parts (e.g. shirt, pants...). The nearest neighbor search is carried out separately for each part. 
+ * The pca basis and the nearest neighbor data are compressed into morph targets.
+ */
+UCLASS()
+class NEARESTNEIGHBORMODEL_API UNearestNeighborModel 
+	: public UMLDeformerMorphModel
+{
+	GENERATED_BODY()
+
+public:
+	UNearestNeighborModel(const FObjectInitializer& ObjectInitializer);
+	virtual void PostLoad() override;
+
+	// UMLDeformerModel overrides.
+	virtual UMLDeformerModelInstance* CreateModelInstance(UMLDeformerComponent* Component);
+	virtual UMLDeformerInputInfo* CreateInputInfo();
+	virtual FString GetDisplayName() const override { return "Nearest Neighbor Model"; }
+	// ~END UMLDeformerModel overrides.
+
+	UFUNCTION(BlueprintPure, Category = "Nearest Neighbor Model")
+	int32 GetNumParts() const { return ClothPartData.Num(); }
+
+	UFUNCTION(BlueprintPure, Category = "Nearest Neighbor Model")
+	int32 GetPartNumVerts(int32 PartId) const { return ClothPartData[PartId].NumVertices; }
+	
+	UFUNCTION(BlueprintPure, Category = "Nearest Neighbor Model")
+	int32 GetPCACoeffStart(int32 PartId) const { return ClothPartData[PartId].PCACoeffStart; }
+
+	UFUNCTION(BlueprintPure, Category = "Nearest Neighbor Model")
+	int32 GetPCACoeffNum(int32 PartId) const { return ClothPartData[PartId].PCACoeffNum; }
+
+	UFUNCTION(BlueprintPure, Category = "Nearest Neighbor Model")
+	int32 GetNumNeighbors(int32 PartId) const { return ClothPartData[PartId].NumNeighbors; }
+
+	float GetDecayFactor() const { return DecayFactor; };
+
+	const TArray<uint32>& PartVertexMap(int32 PartId) const { return ClothPartData[PartId].VertexMap; }
+
+	UFUNCTION(BlueprintPure, Category = "Nearest Neighbor Model")
+	const TArray<float>& PCABasis(int32 PartId) const { return ClothPartData[PartId].PCABasis; }
+
+	TArray<float>& PCABasis(int32 PartId) { return ClothPartData[PartId].PCABasis; }
+
+	UFUNCTION(BlueprintCallable, Category = "Nearest Neighbor Model")
+	void SetPCABasis(int32 PartId, const TArray<float>& PCABasis) { ClothPartData[PartId].PCABasis = PCABasis; }
+
+	UFUNCTION(BlueprintPure, Category = "Nearest Neighbor Model")
+	const TArray<float>& VertexMean(int32 PartId) const { return ClothPartData[PartId].VertexMean; }
+
+	TArray<float>& VertexMean(int32 PartId) { return ClothPartData[PartId].VertexMean; }
+
+	UFUNCTION(BlueprintCallable, Category = "Nearest Neighbor Model")
+	void SetVertexMean(int32 PartId, const TArray<float>& VertexMean) { ClothPartData[PartId].VertexMean = VertexMean; }
+
+	UFUNCTION(BlueprintPure, Category = "Nearest Neighbor Model")
+	const TArray<float>& NeighborCoeffs(int32 PartId) const { return ClothPartData[PartId].NeighborCoeffs; }
+
+	TArray<float>& NeighborCoeffs(int32 PartId) { return ClothPartData[PartId].NeighborCoeffs; }
+
+	UFUNCTION(BlueprintCallable, Category = "Nearest Neighbor Model")
+	void SetNeighborCoeffs(int32 PartId, const TArray<float>& NeighborCoeffs) { ClothPartData[PartId].NeighborCoeffs = NeighborCoeffs; }
+
+
+	UFUNCTION(BlueprintPure, Category = "Nearest Neighbor Model")
+	const TArray<float>& NeighborOffsets(int32 PartId) const { return ClothPartData[PartId].NeighborOffsets; }
+
+	TArray<float>& NeighborOffsets(int32 PartId) { return ClothPartData[PartId].NeighborOffsets; }
+
+	UFUNCTION(BlueprintCallable, Category = "Nearest Neighbor Model")
+	void SetNeighborOffsets(int32 PartId, const TArray<float>& NeighborOffsets) { ClothPartData[PartId].NeighborOffsets = NeighborOffsets; }
+
+	void ClipInputs(float* InputPtr, int NumInputs);
+
+	bool CheckPCAData(int32 PartId) const;
+
+	void InitInputInfo();
+	void InitPreviousWeights();
+
+#define EDITORONLY_PROPERTY_HELPERS(TYPE, NAME) \
+	UFUNCTION(BlueprintPure, Category = "Nearest Neighbor Model")\
+	const TYPE Get##NAME() const { return NAME; }\
+	\
+	static FName Get##NAME##PropertyName() { return GET_MEMBER_NAME_CHECKED(UNearestNeighborModel, NAME); }
+
+#if WITH_EDITORONLY_DATA
+	TObjectPtr<UAnimSequence> GetNearestNeighborSkeletons(int32 PartId);
+	const TObjectPtr<UAnimSequence> GetNearestNeighborSkeletons(int32 PartId) const;
+	TObjectPtr<UGeometryCache> GetNearestNeighborCache(int32 PartId);
+	const TObjectPtr<UGeometryCache> GetNearestNeighborCache(int32 PartId) const;
+	int32 GetNumNeighborsFromGeometryCache(int32 PartId) const;
+
+	uint32 NearestNeighborId(int32 PartId) const { return NearestNeighborIds[PartId]; }
+	void SetNearestNeighborId(int32 PartId, int32 Id) { NearestNeighborIds[PartId] = Id; }
+
+	void UpdateNetworkInputDim();
+	void UpdateNetworkOutputDim();
+	void UpdateClothPartData();
+	void UpdatePCACoeffNums();
+
+	void InvalidateClothPartData() { bClothPartDataValid = false; bNearestNeighborDataValid = false; }
+	void ValidateClothPartData() { bClothPartDataValid = true; }
+	void InvalidateNearestNeighborData() { bNearestNeighborDataValid = false; }
+	void ValidateNearestNeighborData() {bNearestNeighborDataValid = true; }
+	bool IsClothPartDataValid() { return bClothPartDataValid; }
+	bool IsNearestNeighborDataValid() { return bNearestNeighborDataValid; }
+
+
+	EDITORONLY_PROPERTY_HELPERS(int32, InputDim)
+	EDITORONLY_PROPERTY_HELPERS(TArray<int32>, HiddenLayerDims)
+	EDITORONLY_PROPERTY_HELPERS(int32, OutputDim)
+	EDITORONLY_PROPERTY_HELPERS(int32, NumEpochs)
+	EDITORONLY_PROPERTY_HELPERS(int32, BatchSize)
+	EDITORONLY_PROPERTY_HELPERS(float, LearningRate)
+	EDITORONLY_PROPERTY_HELPERS(TArray<FClothPartEditorData>, ClothPartEditorData)
+	EDITORONLY_PROPERTY_HELPERS(TArray<FSkeletonCachePair>, NearestNeighborData)
+#endif
+
+
+protected:
+#if WITH_EDITORONLY_DATA
+	UPROPERTY(VisibleAnywhere, BlueprintReadOnly, Category = "Training Settings")
+	int32 InputDim = 0;
+
+	UPROPERTY(EditAnywhere, BlueprintReadWrite, Category = "Training Settings")
+	TArray<int32> HiddenLayerDims;
+
+	UPROPERTY(VisibleAnywhere, BlueprintReadOnly, Category = "Training Settings")
+	int32 OutputDim = 0;
+
+	UPROPERTY(EditAnywhere, BlueprintReadWrite, Category = "Training Settings", meta = (ClampMin = "1"))
+	int32 NumEpochs = 10000;
+
+	UPROPERTY(EditAnywhere, BlueprintReadWrite, Category = "Training Settings", meta = (ClampMin = "1"))
+	int32 BatchSize = 256;
+
+	UPROPERTY(EditAnywhere, BlueprintReadWrite, Category = "Training Settings", meta = (ClampMin = "0.000001", ClampMax = "1.0"))
+	float LearningRate = 0.001f;
+
+	UPROPERTY()
+	bool bClothPartDataValid = false;
+
+	UPROPERTY(EditAnywhere, BlueprintReadWrite, Category = "Cloth Parts")
+	TArray<FClothPartEditorData> ClothPartEditorData;
+
+	UPROPERTY()
+	bool bNearestNeighborDataValid = false;
+
+	UPROPERTY(EditAnywhere, BlueprintReadWrite, Category = "Nearest Neighbors")
+	TArray<FSkeletonCachePair> NearestNeighborData;
+
+	TArray<uint32> NearestNeighborIds;
+#endif
+
+public:
+	UPROPERTY()
+	TArray<FClothPartData> ClothPartData;
+
+	UPROPERTY(EditAnywhere, BlueprintReadWrite, Category = "Network Inputs")
+	TArray<float> InputsMin;
+
+	UPROPERTY(EditAnywhere, BlueprintReadWrite, Category = "Network Inputs")
+	TArray<float> InputsMax;
+
+	UPROPERTY(EditAnywhere, BlueprintReadWrite, Category = "KMeans Pose Generator")
+	TArray<TObjectPtr<UAnimSequence>> SourceSkeletons;
+
+	UPROPERTY(EditAnywhere, BlueprintReadWrite, Category = "KMeans Pose Generator", meta = (ClampMin = "1"))
+	int32 NumClusters;
+
+	UPROPERTY(EditAnywhere, BlueprintReadWrite, Category = "Nearest Neighbors", META = (ClampMin = "0", ClampMax = "1"))
+	float DecayFactor = 0.85f;
+
+	TArray<float> PreviousWeights; 
+};
diff --git a/Engine/Plugins/Animation/MLDeformer/NearestNeighborModel/Source/NearestNeighborModel/Public/NearestNeighborModelInputInfo.h b/Engine/Plugins/Animation/MLDeformer/NearestNeighborModel/Source/NearestNeighborModel/Public/NearestNeighborModelInputInfo.h
new file mode 100644
--- /dev/null
+++ b/Engine/Plugins/Animation/MLDeformer/NearestNeighborModel/Source/NearestNeighborModel/Public/NearestNeighborModelInputInfo.h
@@ -0,0 +1,26 @@
+// Copyright Epic Games, Inc. All Rights Reserved.
+
+#pragma once
+
+#include "MLDeformerInputInfo.h"
+#include "NearestNeighborModelInputInfo.generated.h"
+
+UCLASS()
+class NEARESTNEIGHBORMODEL_API UNearestNeighborModelInputInfo
+    : public UMLDeformerInputInfo
+{
+    GENERATED_BODY()
+
+public:
+    // UMLDeformerInputInfo overrides
+    virtual int32 CalcNumNeuralNetInputs() const override;
+    virtual void ExtractBoneRotations(USkeletalMeshComponent* SkelMeshComponent, TArray<float>& OutRotations) const override;
+    // ~END UMLDeformerInputInfo overrides
+
+    void InitRefBoneRotations(USkeletalMesh* SkelMesh);
+    void ComputeNetworkInput(const TArray<FQuat>& BoneRotations, float* OutputBuffer, int64 StartIndex = 0) const;
+
+private:
+    UPROPERTY()
+    TArray<FQuat> RefBoneRotations;
+};
diff --git a/Engine/Plugins/Animation/MLDeformer/NearestNeighborModel/Source/NearestNeighborModel/Public/NearestNeighborModelInstance.h b/Engine/Plugins/Animation/MLDeformer/NearestNeighborModel/Source/NearestNeighborModel/Public/NearestNeighborModelInstance.h
new file mode 100644
--- /dev/null
+++ b/Engine/Plugins/Animation/MLDeformer/NearestNeighborModel/Source/NearestNeighborModel/Public/NearestNeighborModelInstance.h
@@ -0,0 +1,27 @@
+// Copyright Epic Games, Inc. All Rights Reserved.
+#pragma once
+
+#include "MLDeformerMorphModelInstance.h"
+#include "NeuralTensor.h"
+#include "NearestNeighborModelInstance.generated.h"
+
+UCLASS()
+class NEARESTNEIGHBORMODEL_API UNearestNeighborModelInstance
+    : public UMLDeformerModelInstance
+{
+    GENERATED_BODY()
+public:
+
+    // UMLDeformerModelInstance overrides
+    virtual void RunNeuralNetwork(float ModelWeight) override;
+    virtual bool SetupNeuralNetworkForFrame() override;
+
+protected:
+    virtual int64 SetBoneTransforms(float* OutputBuffer, int64 OutputBufferSize, int64 StartIndex) override;
+    // ~END UMLDeformerModelInstance overrides
+
+private:
+    void RunNearestNeighborModel(float ModelWeight);
+    int32 FindNearestNeighbor(const FNeuralTensor& PCACoeffTensor, int32 PartId);
+    void UpdateWeight(TArray<float>& MorphWeights, const int32 Index, const float W);
+};
diff --git a/Engine/Plugins/Animation/MLDeformer/NearestNeighborModel/Source/NearestNeighborModel/Public/NearestNeighborModelVizSettings.h b/Engine/Plugins/Animation/MLDeformer/NearestNeighborModel/Source/NearestNeighborModel/Public/NearestNeighborModelVizSettings.h
new file mode 100644
--- /dev/null
+++ b/Engine/Plugins/Animation/MLDeformer/NearestNeighborModel/Source/NearestNeighborModel/Public/NearestNeighborModelVizSettings.h
@@ -0,0 +1,19 @@
+// Copyright Epic Games, Inc. All Rights Reserved.
+
+#pragma once
+
+#include "CoreMinimal.h"
+#include "MLDeformerMorphModelVizSettings.h"
+#include "NearestNeighborModelVizSettings.generated.h"
+
+class UGeometryCache;
+
+/**
+ * The vizualization settings specific to the the vertex delta model.
+ */
+UCLASS()
+class NEARESTNEIGHBORMODEL_API UNearestNeighborModelVizSettings 
+	: public UMLDeformerMorphModelVizSettings
+{
+	GENERATED_BODY()
+};
diff --git a/Engine/Plugins/Animation/MLDeformer/NearestNeighborModel/Source/NearestNeighborModelEditor/NearestNeighborModelEditor.Build.cs b/Engine/Plugins/Animation/MLDeformer/NearestNeighborModel/Source/NearestNeighborModelEditor/NearestNeighborModelEditor.Build.cs
new file mode 100644
--- /dev/null
+++ b/Engine/Plugins/Animation/MLDeformer/NearestNeighborModel/Source/NearestNeighborModelEditor/NearestNeighborModelEditor.Build.cs
@@ -0,0 +1,50 @@
+// Copyright Epic Games, Inc. All Rights Reserved.
+
+using System.IO;
+
+namespace UnrealBuildTool.Rules
+{
+	public class NearestNeighborModelEditor : ModuleRules
+	{
+		public NearestNeighborModelEditor(ReadOnlyTargetRules Target) : base(Target)
+		{
+			PublicDependencyModuleNames.AddRange(
+				new string[]
+				{
+					"Core"
+				}
+			);
+
+			PrivateDependencyModuleNames.AddRange(
+				new string[]
+				{
+					"EditorFramework",
+					"UnrealEd",
+					"CoreUObject",
+					"Engine",
+					"Slate",
+					"SlateCore",
+					"EditorWidgets",
+					"EditorStyle",
+					"GeometryCache",
+					"MLDeformerFramework",
+					"MLDeformerFrameworkEditor",
+					"NearestNeighborModel",
+					"Projects",
+					"PropertyEditor",
+					"NeuralNetworkInference",
+					"ToolWidgets",
+					"ComputeFramework",
+					"RenderCore",
+					"RHI"
+				}
+			);
+
+			PrivateIncludePathModuleNames.AddRange(
+				new string[]
+				{
+				}
+			);
+		}
+	}
+}
diff --git a/Engine/Plugins/Animation/MLDeformer/NearestNeighborModel/Source/NearestNeighborModelEditor/Private/NearestNeighborEditorModel.cpp b/Engine/Plugins/Animation/MLDeformer/NearestNeighborModel/Source/NearestNeighborModelEditor/Private/NearestNeighborEditorModel.cpp
new file mode 100644
--- /dev/null
+++ b/Engine/Plugins/Animation/MLDeformer/NearestNeighborModel/Source/NearestNeighborModelEditor/Private/NearestNeighborEditorModel.cpp
@@ -0,0 +1,391 @@
+// Copyright Epic Games, Inc. All Rights Reserved.
+
+#include "NearestNeighborEditorModel.h"
+#include "NearestNeighborModel.h"
+#include "NearestNeighborTrainingModel.h"
+#include "NearestNeighborModelInputInfo.h"
+#include "NearestNeighborEditorModelActor.h"
+#include "NearestNeighborModelStyle.h"
+#include "NearestNeighborGeomCacheSampler.h"
+#include "MLDeformerEditorToolkit.h"
+#include "NeuralNetwork.h"
+#include "GeometryCache.h"
+#include "GeometryCacheComponent.h"
+#include "Animation/DebugSkelMeshComponent.h"
+#include "Animation/MorphTarget.h"
+#include "Animation/AnimSequence.h"
+
+#define LOCTEXT_NAMESPACE "NearestNeighborEditorModel"
+
+namespace UE::NearestNeighborModel
+{
+	using namespace UE::MLDeformer;
+
+	FMLDeformerEditorModel* FNearestNeighborEditorModel::MakeInstance()
+	{
+		return new FNearestNeighborEditorModel();
+	}
+
+	void FNearestNeighborEditorModel::Init(const InitSettings& InitSettings)
+	{
+		FMLDeformerEditorModel::Init(InitSettings);
+		InitInputInfo(Model->GetInputInfo());
+		InitMorphTargets();
+	}
+
+	FMLDeformerSampler* FNearestNeighborEditorModel::CreateSampler() const
+	{
+		FNearestNeighborGeomCacheSampler* NewSampler = new FNearestNeighborGeomCacheSampler();
+		NewSampler->OnGetGeometryCache().BindLambda([this] { return GetGeomCacheModel()->GetGeometryCache(); });
+		return NewSampler;
+	}
+
+	void FNearestNeighborEditorModel::OnPropertyChanged(FPropertyChangedEvent& PropertyChangedEvent)
+	{
+		const FProperty* Property = PropertyChangedEvent.Property;
+		if (Property == nullptr)
+		{
+			return;
+		}
+
+		FMLDeformerMorphModelEditorModel::OnPropertyChanged(PropertyChangedEvent);
+
+		if (Property->GetFName() == UMLDeformerModel::GetSkeletalMeshPropertyName() ||
+			Property->GetFName() == UMLDeformerModel::GetAnimSequencePropertyName() ||
+			Property->GetFName() == UMLDeformerGeomCacheModel::GetGeometryCachePropertyName() ||
+			Property->GetFName() == UNearestNeighborModel::GetClothPartEditorDataPropertyName()
+			|| (PropertyChangedEvent.MemberProperty != nullptr && PropertyChangedEvent.MemberProperty->GetFName() == UNearestNeighborModel::GetClothPartEditorDataPropertyName()))
+		{
+			GetNearestNeighborModel()->InvalidateClothPartData();
+			GetEditor()->GetModelDetailsView()->ForceRefresh();
+		}
+
+		if (Property->GetFName() == UNearestNeighborModel::GetNearestNeighborDataPropertyName())
+		{
+			GetNearestNeighborModel()->InvalidateNearestNeighborData();
+			GetEditor()->GetModelDetailsView()->ForceRefresh();
+		}
+	}
+
+	void FNearestNeighborEditorModel::OnPostTraining(ETrainingResult TrainingResult, bool bUsePartiallyTrainedWhenAborted)
+	{
+		FMLDeformerMorphModelEditorModel::OnPostTraining(TrainingResult, bUsePartiallyTrainedWhenAborted);
+		GetNearestNeighborModel()->InitPreviousWeights();
+	}
+
+	void FNearestNeighborEditorModel::CreateActors(const TSharedRef<IPersonaPreviewScene>& InPersonaPreviewScene)
+	{
+		FMLDeformerMorphModelEditorModel::CreateActors(InPersonaPreviewScene);
+
+		UWorld* World = InPersonaPreviewScene->GetWorld();
+		CreateNearestNeighborActors(World);
+		EditorWorld = World;
+	}
+
+	void FNearestNeighborEditorModel::Tick(FEditorViewportClient* ViewportClient, float DeltaTime)
+	{
+		FMLDeformerEditorModel::Tick(ViewportClient, DeltaTime);
+		for (FNearestNeighborEditorModelActor* NearestNeighborActor : NearestNeighborActors)
+		{
+			if (NearestNeighborActor)
+			{
+				NearestNeighborActor->TickNearestNeighborActor();
+			}
+		}
+	}
+
+	void FNearestNeighborEditorModel::CreateNearestNeighborActors(UWorld* World, int32 StartIndex)
+	{
+		UNearestNeighborModel* NearestNeighborModel = GetNearestNeighborModel();
+		const int32 NumParts = NearestNeighborModel->GetNumParts();
+
+		if (StartIndex == 0)
+		{
+			NearestNeighborActors.Reset();
+		}
+		NearestNeighborActors.SetNumZeroed(NumParts);
+		for(int32 PartId = StartIndex; PartId < NumParts; PartId++)
+		{
+			UGeometryCache* GeomCache = NearestNeighborModel->GetNearestNeighborCache(PartId);
+			const FLinearColor LabelColor = FNearestNeighborModelEditorStyle::Get().GetColor("NearestNeighborModel.NearestNeighborActors.LabelColor");
+			const FLinearColor WireframeColor = FNearestNeighborModelEditorStyle::Get().GetColor("NearestNeighborModel.NearestNeighborActors.WireframeColor");
+			CreateGeomCacheActor(
+				World, 
+				ActorID_NearestNeighborActors,
+				*FString::Printf(TEXT("NearestNeighbors%d"), PartId), 
+				GeomCache, 
+				LabelColor,					
+				WireframeColor,
+				LOCTEXT("TestNearestNeighborLabelText", "Nearest Neigbors"),
+				false);
+			FNearestNeighborEditorModelActor* NearestNeighborActor = static_cast<FNearestNeighborEditorModelActor*>(EditorActors.Last());
+			NearestNeighborActor->InitNearestNeighborActor(NearestNeighborModel, PartId);
+			NearestNeighborActor->SetMeshOffsetFactor(2.0f);
+			NearestNeighborActors[PartId] = NearestNeighborActor;
+		}
+	}
+
+	FMLDeformerEditorActor* FNearestNeighborEditorModel::CreateEditorActor(const FMLDeformerEditorActor::FConstructSettings& Settings) const
+	{
+		return new FNearestNeighborEditorModelActor(Settings);
+	}
+
+	void FNearestNeighborEditorModel::InitInputInfo(UMLDeformerInputInfo* InputInfo)
+	{
+		FMLDeformerEditorModel::InitInputInfo(InputInfo);
+		UNearestNeighborModel *NearestNeighborModel = static_cast<UNearestNeighborModel*>(Model);
+		NearestNeighborModel->InitInputInfo();
+	}
+
+	ETrainingResult FNearestNeighborEditorModel::Train()
+	{
+		return TrainModel<UNearestNeighborTrainingModel>(this);
+	}
+
+	bool FNearestNeighborEditorModel::SetSamplerPartData(const int32 PartId)
+	{
+		FNearestNeighborGeomCacheSampler* GeomCacheSampler = static_cast<FNearestNeighborGeomCacheSampler*>(GetGeomCacheSampler());
+		UNearestNeighborModel* NearestNeighborModel = GetNearestNeighborModel();
+		if (GeomCacheSampler && NearestNeighborModel && PartId < NearestNeighborModel->GetNumParts())
+		{
+			TObjectPtr<USkeletalMeshComponent> SkeletalMeshComponent = GeomCacheSampler->GetSkeletalMeshComponent();
+			TObjectPtr<UGeometryCacheComponent> GeometryCacheComponent = GeomCacheSampler->GetGeometryCacheComponent();
+
+			const TObjectPtr<UAnimSequence> AnimSequence = NearestNeighborModel->GetNearestNeighborSkeletons(PartId);
+			const TObjectPtr<UGeometryCache> GeometryCache = NearestNeighborModel->GetNearestNeighborCache(PartId);
+
+			if (!SkeletalMeshComponent.IsNull() && !GeometryCacheComponent.IsNull() && !AnimSequence.IsNull() && !GeometryCache.IsNull())
+			{
+				const int32 NumFrames = GeometryCache->GetEndFrame() - GeometryCache->GetStartFrame() + 1;
+				if (AnimSequence->GetDataModel()->GetNumberOfKeys() == NumFrames)
+				{
+					AnimSequence->bUseRawDataOnly = true;
+					AnimSequence->Interpolation = EAnimInterpolationType::Step;
+
+					SkeletalMeshComponent->SetAnimationMode(EAnimationMode::AnimationSingleNode);
+					SkeletalMeshComponent->SetAnimation(AnimSequence);
+					SkeletalMeshComponent->SetPosition(0.0f);
+					SkeletalMeshComponent->SetPlayRate(1.0f);
+					SkeletalMeshComponent->Play(false);
+					SkeletalMeshComponent->RefreshBoneTransforms();
+
+					// assuming MeshMappings do not change
+					GeometryCacheComponent->SetGeometryCache(GeometryCache);
+					GeometryCacheComponent->ResetAnimationTime();
+					GeometryCacheComponent->SetLooping(false);
+					GeometryCacheComponent->SetManualTick(true);
+					GeometryCacheComponent->SetPlaybackSpeed(1.0f);
+					GeometryCacheComponent->Play();
+					GeomCacheSampler->GeneratePartMeshMappings(GetNearestNeighborModel()->PartVertexMap(PartId));
+
+					return true;
+				}
+				else
+				{
+					UE_LOG(LogNearestNeighborModel, Error, TEXT("Part %d frame mismatch: AnimSequence has %d frames and GeometryCache has %d frames"), PartId, AnimSequence->GetDataModel()->GetNumberOfFrames(), NumFrames);
+				}
+			}
+			else
+			{
+				UE_LOG(LogNearestNeighborModel, Error, TEXT("Part %d: AnimSequence or GeometryCache is null"), PartId);
+			}
+		}
+		return false;
+	}
+
+	int32 FNearestNeighborEditorModel::GetNumParts()
+	{
+		UNearestNeighborModel *NearestNeighborModel = GetNearestNeighborModel();
+		if (NearestNeighborModel != nullptr)
+		{
+			return NearestNeighborModel->GetNumParts();
+		}
+		else
+		{
+			return 0;
+		}
+	}
+
+	void FNearestNeighborEditorModel::ResetSamplerData()
+	{
+		FNearestNeighborGeomCacheSampler* GeomCacheSampler = static_cast<FNearestNeighborGeomCacheSampler*>(GetGeomCacheSampler());
+		UNearestNeighborModel* NearestNeighborModel = GetNearestNeighborModel();
+		if (GeomCacheSampler && NearestNeighborModel)
+		{
+			TObjectPtr<USkeletalMeshComponent> SkeletalMeshComponent = GeomCacheSampler->GetSkeletalMeshComponent();
+			const TObjectPtr<UAnimSequence> AnimSequence = NearestNeighborModel->GetAnimSequence();
+
+			if (!SkeletalMeshComponent.IsNull() && !AnimSequence.IsNull())
+			{
+				SkeletalMeshComponent->SetAnimationMode(EAnimationMode::AnimationSingleNode);
+				SkeletalMeshComponent->SetAnimation(AnimSequence);
+				SkeletalMeshComponent->SetPosition(0.0f);
+				SkeletalMeshComponent->SetPlayRate(1.0f);
+				SkeletalMeshComponent->Play(false);
+				SkeletalMeshComponent->RefreshBoneTransforms();
+			}
+
+			GeomCacheSampler->RegisterTargetComponents();
+		}
+	}
+
+	void FNearestNeighborEditorModel::UpdateNearestNeighborActors()
+	{
+		if (EditorWorld != nullptr)
+		{
+			const UNearestNeighborModel* NearestNeighborModel = GetNearestNeighborModel();
+			const int32 TargetNumActors = NearestNeighborModel->GetNumParts();
+			if (NearestNeighborActors.Num() > TargetNumActors)
+			{
+				for (int32 i = NearestNeighborActors.Num() - 1; i >= TargetNumActors; i--)
+				{
+					FNearestNeighborEditorModelActor* EditorActor = NearestNeighborActors[i];
+					NearestNeighborActors.RemoveAt(i);
+					EditorActors.Remove(EditorActor);
+					EditorWorld->RemoveActor(EditorActor->GetActor(), true/*ShouldModifyLevel*/);
+					delete EditorActor;
+				}
+			}
+			if (NearestNeighborActors.Num() < TargetNumActors)
+			{
+				CreateNearestNeighborActors(EditorWorld, NearestNeighborActors.Num());
+			}
+		}
+
+		for (int32 PartId = 0; PartId < NearestNeighborActors.Num(); PartId++)
+		{
+			const TObjectPtr<UGeometryCache> GeometryCache = GetNearestNeighborModel()->GetNearestNeighborCache(PartId);
+			NearestNeighborActors[PartId]->GetGeometryCacheComponent()->SetGeometryCache(GeometryCache);
+		}
+	}
+
+	void FNearestNeighborEditorModel::UpdateNearestNeighborData()
+	{
+		UNearestNeighborModel *NearestNeighborModel = static_cast<UNearestNeighborModel*>(Model);
+		check(NearestNeighborModel != nullptr);
+
+		UNeuralNetwork *NeuralNetwork = NearestNeighborModel->GetNeuralNetwork();
+		if(NeuralNetwork && NeuralNetwork->IsLoaded())
+		{
+			const FString SavePath = FPaths::ProjectIntermediateDir() + "NearestNeighborModel/NearestNeighborModel.onnx";
+			if (!FPaths::FileExists(SavePath))
+			{
+				UE_LOG(LogNearestNeighborModel, Display, TEXT("Saving to %s"), *SavePath);
+				NeuralNetwork->Save(SavePath);
+			}
+		}
+		else
+		{
+			UE_LOG(LogNearestNeighborModel, Warning, TEXT("Network is not available. Nothing will be done."));
+			return;
+		}
+
+		FNearestNeighborGeomCacheSampler* GeomCacheSampler = static_cast<FNearestNeighborGeomCacheSampler*>(GetGeomCacheSampler());
+		if (GeomCacheSampler)
+		{
+			TObjectPtr<USkeletalMeshComponent> SkeletalMeshComponent = GeomCacheSampler->GetSkeletalMeshComponent();
+			TObjectPtr<UGeometryCacheComponent> GeometryCacheComponent = GeomCacheSampler->GetGeometryCacheComponent();
+
+			if(!SkeletalMeshComponent.IsNull() && !GeometryCacheComponent.IsNull())
+			{
+				UNearestNeighborTrainingModel *TrainingModel = InitTrainingModel<UNearestNeighborTrainingModel>(this);
+				check(TrainingModel != nullptr);
+				TrainingModel->UpdateNearestNeighborData();
+				ResetSamplerData();
+				UpdateNearestNeighborActors();
+				NearestNeighborModel->ValidateNearestNeighborData();
+			}
+		}
+	}
+
+	void FNearestNeighborEditorModel::KMeansClusterPoses()
+	{
+		UNearestNeighborModel *NearestNeighborModel = static_cast<UNearestNeighborModel*>(Model);
+		check(NearestNeighborModel != nullptr);
+
+		UNeuralNetwork *NeuralNetwork = NearestNeighborModel->GetNeuralNetwork();
+		if(NeuralNetwork && NeuralNetwork->IsLoaded())
+		{
+			const FString SavePath = FPaths::ProjectIntermediateDir() + "NearestNeighborModel/NearestNeighborModel.onnx";
+			if (!FPaths::FileExists(SavePath))
+			{
+				UE_LOG(LogNearestNeighborModel, Display, TEXT("Saving to %s"), *SavePath);
+				NeuralNetwork->Save(SavePath);
+			}
+		}
+		else
+		{
+			UE_LOG(LogNearestNeighborModel, Warning, TEXT("Network is not available. Nothing will be done."));
+			return;
+		}
+
+		UNearestNeighborTrainingModel *TrainingModel = InitTrainingModel<UNearestNeighborTrainingModel>(this);
+		check(TrainingModel != nullptr);
+		TrainingModel->KmeansClusterPoses();
+		ResetSamplerData();
+	}
+
+	void FNearestNeighborEditorModel::AddFloatArrayToDeltaArray(const TArray<float>& FloatArr, const TArray<uint32>& VertexMap, TArray<FVector3f>& DeltaArr, int32 DeltaArrayOffset, float ScaleFactor)
+	{
+		const int32 NumBaseMeshVerts = Model->GetNumBaseMeshVerts();
+		const UNearestNeighborModel* NearestNeighborModel = GetNearestNeighborModel();
+		const int32 PartNumVerts = VertexMap.Num();
+		const int32 NumShapes = FloatArr.Num() / (PartNumVerts * 3);
+		if (DeltaArrayOffset < 0)
+		{
+			DeltaArrayOffset = DeltaArr.Num();
+		}
+		DeltaArr.SetNumZeroed(FMath::Max(DeltaArrayOffset + NumShapes * NumBaseMeshVerts, DeltaArr.Num()), false);
+
+		for(int32 ShapeId = 0; ShapeId < NumShapes; ShapeId++)
+		{
+			for(int32 VertexId = 0; VertexId < PartNumVerts; VertexId++)
+			{
+				const int32 DeltaId = ShapeId * NumBaseMeshVerts + VertexMap[VertexId];
+				const int32 FloatId = (ShapeId * PartNumVerts + VertexId) * 3;
+				DeltaArr[DeltaArrayOffset + DeltaId] = FVector3f(FloatArr[FloatId], FloatArr[FloatId + 1], FloatArr[FloatId + 2]) * ScaleFactor;
+			}
+		}
+	}
+
+	void FNearestNeighborEditorModel::InitMorphTargets()
+	{
+		UNearestNeighborModel* NearestNeighborModel = GetNearestNeighborModel();
+		if (NearestNeighborModel->GetNumBaseMeshVerts() == 0)
+		{
+			return;
+		}
+		TArray<FVector3f> Deltas;
+		Deltas.Reset();
+
+		int32 NumPCACoeff = 0;
+		for (int32 PartId = 0; PartId < NearestNeighborModel->GetNumParts(); PartId++)
+		{
+			NumPCACoeff += NearestNeighborModel->GetPCACoeffNum(PartId);
+		}
+		Deltas.Reserve((1 + NumPCACoeff) * NearestNeighborModel->GetNumBaseMeshVerts());
+
+		for (int32 PartId = 0; PartId < NearestNeighborModel->GetNumParts(); PartId++)
+		{
+			const TArray<uint32>& VertexMap = NearestNeighborModel->PartVertexMap(PartId);
+			AddFloatArrayToDeltaArray(NearestNeighborModel->ClothPartData[PartId].VertexMean, VertexMap, Deltas, 0);
+			AddFloatArrayToDeltaArray(NearestNeighborModel->ClothPartData[PartId].PCABasis, VertexMap, Deltas);
+		}
+
+		for (int32 PartId = 0; PartId < NearestNeighborModel->GetNumParts(); PartId++)
+		{
+			const TArray<uint32>& VertexMap = NearestNeighborModel->PartVertexMap(PartId);
+			AddFloatArrayToDeltaArray(NearestNeighborModel->ClothPartData[PartId].NeighborOffsets, VertexMap, Deltas);
+		}
+
+		const int32 LOD = 0;
+		TArray<UMorphTarget*> MorphTargets;
+		CreateEngineMorphTargets(MorphTargets, Deltas, FString("NNMorphTarget_"), LOD, NearestNeighborModel->GetMorphTargetDeltaThreshold());
+		check(NearestNeighborModel->GetMorphTargetSet().IsValid());
+		FMorphTargetVertexInfoBuffers& MorphBuffers = NearestNeighborModel->GetMorphTargetSet()->MorphBuffers;
+		CompressEngineMorphTargets(MorphBuffers, MorphTargets, LOD, NearestNeighborModel->GetMorphTargetErrorTolerance());
+		NearestNeighborModel->SetMorphTargetDeltas(Deltas);
+	}
+}	// namespace UE::NearestNeighborModel
+
+#undef LOCTEXT_NAMESPACE
diff --git a/Engine/Plugins/Animation/MLDeformer/NearestNeighborModel/Source/NearestNeighborModelEditor/Private/NearestNeighborEditorModelActor.cpp b/Engine/Plugins/Animation/MLDeformer/NearestNeighborModel/Source/NearestNeighborModelEditor/Private/NearestNeighborEditorModelActor.cpp
new file mode 100644
--- /dev/null
+++ b/Engine/Plugins/Animation/MLDeformer/NearestNeighborModel/Source/NearestNeighborModelEditor/Private/NearestNeighborEditorModelActor.cpp
@@ -0,0 +1,29 @@
+// Copyright Epic Games, Inc. All Rights Reserved.
+
+#include "NearestNeighborEditorModelActor.h"
+#include "GeometryCacheComponent.h"
+
+namespace UE::NearestNeighborModel
+{
+	using namespace UE::MLDeformer;
+
+	FNearestNeighborEditorModelActor::FNearestNeighborEditorModelActor(const FConstructSettings& Settings)
+		: FMLDeformerGeomCacheActor(Settings)
+	{
+	}
+
+	void FNearestNeighborEditorModelActor::InitNearestNeighborActor(UNearestNeighborModel* InNearestNeighborModel, const int32 InPartId)
+	{
+		NearestNeighborModel = InNearestNeighborModel;
+		PartId = InPartId;
+	}
+
+	void FNearestNeighborEditorModelActor::TickNearestNeighborActor()
+	{
+		if (GeomCacheComponent && GeomCacheComponent->GetGeometryCache() && NearestNeighborModel)
+		{
+			GeomCacheComponent->SetManualTick(true);
+			GeomCacheComponent->TickAtThisTime(GeomCacheComponent->GetTimeAtFrame(NearestNeighborModel->NearestNeighborId(PartId)), false, false, false);
+		}
+	}
+}	// namespace UE::NearestNeighborModel
diff --git a/Engine/Plugins/Animation/MLDeformer/NearestNeighborModel/Source/NearestNeighborModelEditor/Private/NearestNeighborGeomCacheSampler.cpp b/Engine/Plugins/Animation/MLDeformer/NearestNeighborModel/Source/NearestNeighborModelEditor/Private/NearestNeighborGeomCacheSampler.cpp
new file mode 100644
--- /dev/null
+++ b/Engine/Plugins/Animation/MLDeformer/NearestNeighborModel/Source/NearestNeighborModelEditor/Private/NearestNeighborGeomCacheSampler.cpp
@@ -0,0 +1,158 @@
+// Copyright Epic Games, Inc. All Rights Reserved.
+
+#include "NearestNeighborGeomCacheSampler.h"
+#include "NearestNeighborModel.h"
+#include "Engine/SkeletalMesh.h"
+#include "Animation/DebugSkelMeshComponent.h"
+#include "Animation/AnimSequence.h"
+#include "Rendering/SkeletalMeshModel.h"
+#include "Rendering/SkeletalMeshLODModel.h"
+#include "Rendering/SkeletalMeshLODRenderData.h"
+#include "GeometryCacheComponent.h"
+#include "GeometryCache.h"
+#include "GeometryCacheMeshData.h"
+#include "GeometryCacheTrack.h"
+
+using namespace UE::MLDeformer;
+namespace UE::NearestNeighborModel
+{
+	void FNearestNeighborGeomCacheSampler::SamplePart(int32 InAnimFrameIndex, const TArray<uint32>& VertexMap)
+	{
+		FMLDeformerSampler::Sample(InAnimFrameIndex);
+		USkeletalMesh* SkeletalMesh = SkeletalMeshComponent.Get() ? SkeletalMeshComponent->GetSkeletalMeshAsset() : nullptr;
+		UGeometryCache* GeometryCache = GeometryCacheComponent.Get() ? GeometryCacheComponent->GetGeometryCache() : nullptr;
+		if (SkeletalMeshComponent && SkeletalMesh && GeometryCacheComponent && GeometryCache)
+		{
+			const FTransform& AlignmentTransform = Model->GetAlignmentTransform();
+			FSkeletalMeshModel* ImportedModel = SkeletalMesh->GetImportedModel();
+			const int32 LODIndex = 0;
+			const FSkeletalMeshLODModel& LODModel = ImportedModel->LODModels[LODIndex];
+			const TArray<FSkelMeshImportedMeshInfo>& SkelMeshInfos = LODModel.ImportedMeshInfos;
+			check(SkelMeshInfos.Num() == 1);
+			const FSkelMeshImportedMeshInfo& MeshInfo = SkelMeshInfos[0]; 
+			check(MeshInfo.StartImportedVertex == 0);
+			check(GeometryCache->Tracks.Num() == 1);
+			UGeometryCacheTrack* Track = GeometryCache->Tracks[0];
+			GeomCacheMeshDatas.Reset(1);
+			GeomCacheMeshDatas.AddDefaulted(1);
+			FGeometryCacheMeshData& GeomCacheMeshData = GeomCacheMeshDatas[0];
+			check(MeshMappings.Num() == 1);
+			const UE::MLDeformer::FMLDeformerGeomCacheMeshMapping& MeshMapping = MeshMappings[0]; 
+
+			if (!Track->GetMeshDataAtTime(SampleTime, GeomCacheMeshData))
+			{
+				UE_LOG(LogNearestNeighborModel, Error, TEXT("FNearestNeighborGeomCacheSampler::SamplePart: Track cannot get mesh delta at frame %d"), InAnimFrameIndex);
+			}
+
+			// Calculate the vertex deltas.
+			const FSkeletalMeshLODRenderData& SkelMeshLODData = SkeletalMesh->GetResourceForRendering()->LODRenderData[LODIndex];
+			const FSkinWeightVertexBuffer& SkinWeightBuffer = *SkeletalMeshComponent->GetSkinWeightBuffer(LODIndex);
+
+			const int32 NumPartVerts = VertexMap.Num();
+			PartVertexDeltas.Reset();
+			PartVertexDeltas.SetNum(NumPartVerts * 3);
+
+			for(int32 PartVertexIndex = 0; PartVertexIndex < NumPartVerts; PartVertexIndex++)
+			{
+				const int32 VertexIndex = VertexMap[PartVertexIndex];
+				const int32 SkinnedVertexIndex = MeshInfo.StartImportedVertex + VertexIndex;
+				const int32 GeomCacheVertexIndex = MeshMapping.SkelMeshToTrackVertexMap[PartVertexIndex];
+
+				if (GeomCacheVertexIndex != INDEX_NONE && GeomCacheMeshData.Positions.IsValidIndex(GeomCacheVertexIndex))
+				{
+					FVector3f Delta = FVector3f::ZeroVector;
+
+					const int32 ArrayIndex = 3 * PartVertexIndex;
+					// Calculate the inverse skinning transform for this vertex.
+					const int32 RenderVertexIndex = MeshMapping.ImportedVertexToRenderVertexMap[PartVertexIndex];
+					if (RenderVertexIndex != INDEX_NONE)
+					{
+						const FMatrix44f InvSkinningTransform = CalcInverseSkinningTransform(RenderVertexIndex, SkelMeshLODData, SkinWeightBuffer);
+
+						// Calculate the pre-skinning data.
+						const FVector3f UnskinnedPosition = SkelMeshLODData.StaticVertexBuffers.PositionVertexBuffer.VertexPosition(RenderVertexIndex);
+						const FVector3f GeomCacheVertexPos = (FVector3f)AlignmentTransform.TransformPosition((FVector)GeomCacheMeshData.Positions[GeomCacheVertexIndex]);
+						const FVector3f PreSkinningTargetPos = InvSkinningTransform.TransformPosition(GeomCacheVertexPos);
+						Delta = PreSkinningTargetPos - UnskinnedPosition;
+					}
+
+					PartVertexDeltas[ArrayIndex] = Delta.X;
+					PartVertexDeltas[ArrayIndex + 1] = Delta.Y;
+					PartVertexDeltas[ArrayIndex + 2] = Delta.Z;
+				}
+			}
+		}
+	}
+
+
+	void FNearestNeighborGeomCacheSampler::GeneratePartMeshMappings(const TArray<uint32>& VertexMap)
+	{
+		USkeletalMesh* SkeletalMesh = SkeletalMeshComponent.Get() ? SkeletalMeshComponent->GetSkeletalMeshAsset() : nullptr;
+		UGeometryCache* GeometryCache = GeometryCacheComponent.Get() ? GeometryCacheComponent->GetGeometryCache() : nullptr;
+		// TODO: make this more general
+		if (SkeletalMeshComponent && SkeletalMesh && GeometryCacheComponent && GeometryCache)
+		{
+			FSkeletalMeshModel* ImportedModel = SkeletalMesh->GetImportedModel();
+			check(ImportedModel);
+			check(ImportedModel->LODModels[0].ImportedMeshInfos.Num() == 1);
+			check(ImportedModel->LODModels[0].ImportedMeshInfos[0].StartImportedVertex == 0);
+
+			check(GeometryCache->Tracks.Num() == 1);
+			UGeometryCacheTrack* Track = GeometryCache->Tracks[0];
+
+			MeshMappings.Reset();
+			MeshMappings.AddDefaulted();
+			UE::MLDeformer::FMLDeformerGeomCacheMeshMapping& Mapping = MeshMappings.Last();
+			Mapping.MeshIndex = 0;
+			Mapping.TrackIndex = 0;
+
+			FGeometryCacheMeshData GeomCacheMeshData;
+			if (!Track->GetMeshDataAtTime(0.0f/*SampleTime*/, GeomCacheMeshData))
+			{
+				return;
+			}
+
+			const int32 NumPartVerts = VertexMap.Num();
+
+			Mapping.SkelMeshToTrackVertexMap.AddUninitialized(NumPartVerts);
+			Mapping.ImportedVertexToRenderVertexMap.AddUninitialized(NumPartVerts);
+
+			for(int32 PartVertexIndex = 0; PartVertexIndex < NumPartVerts; PartVertexIndex++)
+			{
+				Mapping.SkelMeshToTrackVertexMap[PartVertexIndex] = GeomCacheMeshData.ImportedVertexNumbers.Find(PartVertexIndex);
+				const int32 VertexIndex = VertexMap[PartVertexIndex];
+				const int32 RenderVertexIndex = ImportedModel->LODModels[0].MeshToImportVertexMap.Find(VertexIndex);
+				Mapping.ImportedVertexToRenderVertexMap[PartVertexIndex] = RenderVertexIndex;
+			}
+		}
+	}
+
+	void FNearestNeighborGeomCacheSampler::SampleKMeansAnim(const int32 SkeletonId)
+	{
+		UNearestNeighborModel* NearestNeighborModel = static_cast<UNearestNeighborModel*>(Model);
+		if (NearestNeighborModel && SkeletonId < NearestNeighborModel->SourceSkeletons.Num() && !SkeletalMeshComponent.IsNull())
+		{
+			const TObjectPtr<UAnimSequence> AnimSequence = NearestNeighborModel->SourceSkeletons[SkeletonId];
+			SkeletalMeshComponent->SetAnimationMode(EAnimationMode::AnimationSingleNode);
+			SkeletalMeshComponent->SetAnimation(AnimSequence);
+			SkeletalMeshComponent->SetPosition(0.0f);
+			SkeletalMeshComponent->SetPlayRate(1.0f);
+			SkeletalMeshComponent->Play(false);
+			SkeletalMeshComponent->RefreshBoneTransforms();
+		}
+	}
+
+	void FNearestNeighborGeomCacheSampler::SampleKMeansFrame(const int32 Frame)
+	{
+		AnimFrameIndex = Frame;
+		SampleTime = GetTimeAtFrame(Frame);
+
+		const USkeletalMesh* SkeletalMesh = SkeletalMeshComponent->GetSkeletalMeshAsset();
+		if (SkeletalMeshComponent && SkeletalMesh)
+		{
+			UpdateSkeletalMeshComponent();
+			UpdateBoneRotations();
+			UpdateCurveValues();
+		}
+	}
+};
diff --git a/Engine/Plugins/Animation/MLDeformer/NearestNeighborModel/Source/NearestNeighborModelEditor/Private/NearestNeighborModelDetails.cpp b/Engine/Plugins/Animation/MLDeformer/NearestNeighborModel/Source/NearestNeighborModelEditor/Private/NearestNeighborModelDetails.cpp
new file mode 100644
--- /dev/null
+++ b/Engine/Plugins/Animation/MLDeformer/NearestNeighborModel/Source/NearestNeighborModelEditor/Private/NearestNeighborModelDetails.cpp
@@ -0,0 +1,139 @@
+// Copyright Epic Games, Inc. All Rights Reserved.
+
+#include "NearestNeighborModelDetails.h"
+#include "NearestNeighborEditorModel.h"
+#include "NearestNeighborModel.h"
+#include "MLDeformerEditorToolkit.h"
+#include "DetailLayoutBuilder.h"
+#include "DetailCategoryBuilder.h"
+#include "DetailWidgetRow.h"
+#include "IDetailsView.h"
+#include "Widgets/Input/SButton.h"
+
+#define LOCTEXT_NAMESPACE "NearestNeighborModelDetails"
+
+namespace UE::NearestNeighborModel
+{
+	using namespace UE::MLDeformer;
+
+	TSharedRef<IDetailCustomization> FNearestNeighborModelDetails::MakeInstance()
+	{
+		return MakeShareable(new FNearestNeighborModelDetails());
+	}
+
+	bool FNearestNeighborModelDetails::UpdateMemberPointers(const TArray<TWeakObjectPtr<UObject>>& Objects)
+	{
+		if (!FMLDeformerGeomCacheModelDetails::UpdateMemberPointers(Objects))
+		{
+			return false;
+		}
+
+		NearestNeighborModel = static_cast<UNearestNeighborModel*>(Model);
+		check(NearestNeighborModel);
+		NearestNeighborEditorModel = static_cast<FNearestNeighborEditorModel*>(EditorModel);
+
+		return (NearestNeighborModel != nullptr && NearestNeighborEditorModel != nullptr);
+	}
+
+	void FNearestNeighborModelDetails::CreateCategories()
+	{
+		FMLDeformerMorphModelDetails::CreateCategories();
+		ClothPartCategoryBuilder = &DetailLayoutBuilder->EditCategory("Cloth Parts", FText::GetEmpty(), ECategoryPriority::Important);
+		NearestNeighborCategoryBuilder = &DetailLayoutBuilder->EditCategory("Nearest Neighbors", FText::GetEmpty(), ECategoryPriority::Important);
+		KMeansCategoryBuilder = &DetailLayoutBuilder->EditCategory("KMeans Pose Generator", FText::GetEmpty(), ECategoryPriority::Important);
+	}
+
+	void FNearestNeighborModelDetails::CustomizeDetails(IDetailLayoutBuilder& DetailBuilder)
+	{
+		// Create all the detail categories and add the properties of the base class.
+		FMLDeformerMorphModelDetails::CustomizeDetails(DetailBuilder);
+
+		// Training settings.
+		TrainingSettingsCategoryBuilder->AddProperty(UNearestNeighborModel::GetInputDimPropertyName());
+		TrainingSettingsCategoryBuilder->AddProperty(UNearestNeighborModel::GetHiddenLayerDimsPropertyName());
+		TrainingSettingsCategoryBuilder->AddProperty(UNearestNeighborModel::GetOutputDimPropertyName());
+		TrainingSettingsCategoryBuilder->AddProperty(UNearestNeighborModel::GetNumEpochsPropertyName());
+		TrainingSettingsCategoryBuilder->AddProperty(UNearestNeighborModel::GetBatchSizePropertyName());
+		TrainingSettingsCategoryBuilder->AddProperty(UNearestNeighborModel::GetLearningRatePropertyName());
+
+		// Cloth part settings
+		ClothPartCategoryBuilder->AddProperty(UNearestNeighborModel::GetClothPartEditorDataPropertyName());
+		FText ButtonText = NearestNeighborModel->IsClothPartDataValid() ? LOCTEXT("Update", "Update") : LOCTEXT("Update *", "Update *");
+		ClothPartCategoryBuilder->AddCustomRow(FText::FromString(""))
+			.WholeRowContent()
+			[
+				SNew(SButton)
+				.Text(ButtonText)
+				.HAlign(HAlign_Center)
+				.OnClicked_Lambda([this]
+				{
+					if (NearestNeighborModel != nullptr)
+					{
+						NearestNeighborModel->UpdateClothPartData();
+						if (NearestNeighborEditorModel != nullptr)
+						{
+							NearestNeighborEditorModel->UpdateNearestNeighborActors();
+						}
+						EditorModel->GetEditor()->GetModelDetailsView()->ForceRefresh();
+					}
+					return FReply::Handled();
+				})
+			];
+
+		// Nearest Neighbor settings
+		NearestNeighborCategoryBuilder->AddProperty(GET_MEMBER_NAME_STRING_CHECKED(UNearestNeighborModel, DecayFactor));
+		NearestNeighborCategoryBuilder->AddProperty(UNearestNeighborModel::GetNearestNeighborDataPropertyName());
+		ButtonText = NearestNeighborModel->IsNearestNeighborDataValid() ? LOCTEXT("Update", "Update") : LOCTEXT("Update *", "Update *");
+
+		NearestNeighborCategoryBuilder->AddCustomRow(FText::FromString(""))
+			.WholeRowContent()
+			[
+				SNew(SButton)
+				.Text(ButtonText)
+				.HAlign(HAlign_Center)
+				.OnClicked_Lambda([this]
+				{
+					if (NearestNeighborEditorModel != nullptr)
+					{
+						NearestNeighborEditorModel->UpdateNearestNeighborData();
+						EditorModel->GetEditor()->GetModelDetailsView()->ForceRefresh();
+					}
+					return FReply::Handled();
+				})
+			];
+
+		KMeansCategoryBuilder->AddProperty(GET_MEMBER_NAME_STRING_CHECKED(UNearestNeighborModel, SourceSkeletons));
+		KMeansCategoryBuilder->AddProperty(GET_MEMBER_NAME_STRING_CHECKED(UNearestNeighborModel, NumClusters));
+		KMeansCategoryBuilder->AddCustomRow(FText::FromString(""))
+			.WholeRowContent()
+			[
+				SNew(SButton)
+				.Text(FText::FromString("Cluster"))
+				.HAlign(HAlign_Center)
+				.OnClicked_Lambda([this]
+				{
+					NearestNeighborEditorModel->KMeansClusterPoses();
+					return FReply::Handled();
+				})
+			];
+
+		MorphTargetCategoryBuilder->AddCustomRow(FText::FromString(""))
+			.WholeRowContent()
+			[
+				SNew(SButton)
+				.Text(FText::FromString("Update"))
+				.HAlign(HAlign_Center)
+				.OnClicked_Lambda([this]
+				{
+					if (NearestNeighborEditorModel != nullptr)
+					{
+						NearestNeighborEditorModel->InitMorphTargets();
+						EditorModel->GetEditor()->GetModelDetailsView()->ForceRefresh();
+					}
+					return FReply::Handled();
+				})
+			];
+	}
+}	// namespace UE::NearestNeighborModel
+
+#undef LOCTEXT_NAMESPACE
diff --git a/Engine/Plugins/Animation/MLDeformer/NearestNeighborModel/Source/NearestNeighborModelEditor/Private/NearestNeighborModelEditorModule.cpp b/Engine/Plugins/Animation/MLDeformer/NearestNeighborModel/Source/NearestNeighborModelEditor/Private/NearestNeighborModelEditorModule.cpp
new file mode 100644
--- /dev/null
+++ b/Engine/Plugins/Animation/MLDeformer/NearestNeighborModel/Source/NearestNeighborModelEditor/Private/NearestNeighborModelEditorModule.cpp
@@ -0,0 +1,67 @@
+// Copyright Epic Games, Inc. All Rights Reserved.
+
+#include "NearestNeighborModelVizSettingsDetails.h"
+#include "NearestNeighborModelDetails.h"
+#include "NearestNeighborEditorModel.h"
+#include "NearestNeighborModel.h"
+#include "MLDeformerEditorModule.h"
+#include "Modules/ModuleManager.h"
+#include "EditorModeRegistry.h"
+
+#define LOCTEXT_NAMESPACE "NearestNeighborModelEditorModule"
+
+namespace UE::NearestNeighborModel
+{
+	class FNearestNeighborModelEditorModule
+		: public IModuleInterface
+	{
+	public:
+		// IModuleInterface overrides.
+		void StartupModule() override;
+		void ShutdownModule() override;
+		// ~END IModuleInterface overrides.
+	};
+}
+IMPLEMENT_MODULE(UE::NearestNeighborModel::FNearestNeighborModelEditorModule, NearestNeighborModelEditor)
+
+namespace UE::NearestNeighborModel
+{
+	using namespace UE::MLDeformer;
+
+	void FNearestNeighborModelEditorModule::StartupModule()
+	{
+		// Register object detail customizations.
+		FPropertyEditorModule& PropertyModule = FModuleManager::LoadModuleChecked<FPropertyEditorModule>("PropertyEditor");
+		PropertyModule.RegisterCustomClassLayout("NearestNeighborModelVizSettings", FOnGetDetailCustomizationInstance::CreateStatic(&FNearestNeighborModelVizSettingsDetails::MakeInstance));
+		PropertyModule.RegisterCustomClassLayout("NearestNeighborModel", FOnGetDetailCustomizationInstance::CreateStatic(&FNearestNeighborModelDetails::MakeInstance));
+		PropertyModule.NotifyCustomizationModuleChanged();
+
+		// Register our custom ML deformer model to the model registry in the ML Deformer Framework.
+		FMLDeformerEditorModule& EditorModule = FModuleManager::LoadModuleChecked<FMLDeformerEditorModule>("MLDeformerFrameworkEditor");
+		FMLDeformerEditorModelRegistry& ModelRegistry = EditorModule.GetModelRegistry();
+		ModelRegistry.RegisterEditorModel(UNearestNeighborModel::StaticClass(), FOnGetEditorModelInstance::CreateStatic(&FNearestNeighborEditorModel::MakeInstance));
+	}
+
+	void FNearestNeighborModelEditorModule::ShutdownModule()
+	{
+		// Unregister our ML Deformer model.
+		if (FModuleManager::Get().IsModuleLoaded(TEXT("MLDeformerFrameworkEditor")))
+		{
+			FMLDeformerEditorModule& EditorModule = FModuleManager::GetModuleChecked<FMLDeformerEditorModule>("MLDeformerFrameworkEditor");
+			FMLDeformerEditorModelRegistry& ModelRegistry = EditorModule.GetModelRegistry();
+			ModelRegistry.UnregisterEditorModel(UNearestNeighborModel::StaticClass());
+		}
+
+		// Unregister object detail customizations for this model.
+		if (FModuleManager::Get().IsModuleLoaded(TEXT("PropertyEditor")))
+		{
+			FPropertyEditorModule& PropertyModule = FModuleManager::GetModuleChecked<FPropertyEditorModule>(TEXT("PropertyEditor"));
+			PropertyModule.UnregisterCustomClassLayout(TEXT("NearestNeighborModelVizSettings"));
+			PropertyModule.UnregisterCustomClassLayout(TEXT("NearestNeighborModel"));
+			PropertyModule.NotifyCustomizationModuleChanged();
+		}
+	}
+
+}	// namespace UE::NearestNeighborModel
+
+#undef LOCTEXT_NAMESPACE
diff --git a/Engine/Plugins/Animation/MLDeformer/NearestNeighborModel/Source/NearestNeighborModelEditor/Private/NearestNeighborModelStyle.cpp b/Engine/Plugins/Animation/MLDeformer/NearestNeighborModel/Source/NearestNeighborModelEditor/Private/NearestNeighborModelStyle.cpp
new file mode 100644
--- /dev/null
+++ b/Engine/Plugins/Animation/MLDeformer/NearestNeighborModel/Source/NearestNeighborModelEditor/Private/NearestNeighborModelStyle.cpp
@@ -0,0 +1,40 @@
+// Copyright Epic Games, Inc. All Rights Reserved.
+
+#include "NearestNeighborModelStyle.h"
+#include "Styling/SlateStyle.h"
+#include "Styling/SlateStyleRegistry.h"
+#include "Styling/SlateStyleMacros.h"
+#include "Styling/SlateTypes.h"
+#include "Math/Color.h"
+#include "Interfaces/IPluginManager.h"
+#include "Misc/Paths.h"
+
+namespace UE::NearestNeighborModel
+{
+	#define MLDEFORMER_IMAGE_BRUSH(RelativePath, ...) FSlateImageBrush(RootToContentDir(RelativePath, TEXT(".png")), __VA_ARGS__)
+
+	FNearestNeighborModelEditorStyle::FNearestNeighborModelEditorStyle() : FSlateStyleSet("NearestNeighborModelEditorStyle")
+	{
+		const FString ResourceDir = FPaths::Combine(IPluginManager::Get().FindPlugin(TEXT("NearestNeighborModel"))->GetBaseDir(), TEXT("Resources"));
+		SetContentRoot(ResourceDir);
+
+		// Colors and sizes.
+		Set("NearestNeighborModel.NearestNeighborActors.WireframeColor", FLinearColor(1.0f, 1.0f, 0.0f));
+		Set("NearestNeighborModel.NearestNeighborActors.LabelColor", FLinearColor(1.0f, 1.0f, 0.0f));
+
+		FSlateStyleRegistry::RegisterSlateStyle(*this);
+	}
+
+	FNearestNeighborModelEditorStyle::~FNearestNeighborModelEditorStyle()
+	{
+		FSlateStyleRegistry::UnRegisterSlateStyle(*this);
+	}
+
+	FNearestNeighborModelEditorStyle& FNearestNeighborModelEditorStyle::Get()
+	{
+		static FNearestNeighborModelEditorStyle Inst;
+		return Inst;
+	}
+}	// namespace UE::NearestNeighborModel
+
+#undef MLDEFORMER_IMAGE_BRUSH
diff --git a/Engine/Plugins/Animation/MLDeformer/NearestNeighborModel/Source/NearestNeighborModelEditor/Private/NearestNeighborTrainingModel.cpp b/Engine/Plugins/Animation/MLDeformer/NearestNeighborModel/Source/NearestNeighborModelEditor/Private/NearestNeighborTrainingModel.cpp
new file mode 100644
--- /dev/null
+++ b/Engine/Plugins/Animation/MLDeformer/NearestNeighborModel/Source/NearestNeighborModelEditor/Private/NearestNeighborTrainingModel.cpp
@@ -0,0 +1,88 @@
+// Copyright Epic Games, Inc. All Rights Reserved.
+
+#include "NearestNeighborTrainingModel.h"
+#include "NearestNeighborModel.h"
+#include "NearestNeighborEditorModel.h"
+#include "NearestNeighborGeomCacheSampler.h"
+#include "Animation/AnimSequence.h"
+
+#define LOCTEXT_NAMESPACE "NearestNeighborTrainingModel"
+using namespace UE::NearestNeighborModel;
+
+void UNearestNeighborTrainingModel::Init(UE::MLDeformer::FMLDeformerEditorModel* InEditorModel)
+{
+	UMLDeformerTrainingModel::Init(InEditorModel);
+	check(InEditorModel != nullptr);
+	check(InEditorModel->GetModel() != nullptr);
+	NearestNeighborModel = static_cast<UNearestNeighborModel*>(InEditorModel->GetModel());
+}
+
+UNearestNeighborModel* UNearestNeighborTrainingModel::GetNearestNeighborModel() const
+{
+	return Cast<UNearestNeighborModel>(GetModel());
+}
+
+UE::NearestNeighborModel::FNearestNeighborEditorModel* UNearestNeighborTrainingModel::GetNearestNeighborEditorModel() const
+{
+	return static_cast<UE::NearestNeighborModel::FNearestNeighborEditorModel*>(EditorModel);
+}
+
+const TArray<int32> UNearestNeighborTrainingModel::GetPartVertexMap(const int32 PartId) const
+{
+	const TArray<uint32>& VertexMap = NearestNeighborModel->PartVertexMap(PartId);
+	return TArray<int32>((int32*)VertexMap.GetData(), VertexMap.Num());
+}
+
+void UNearestNeighborTrainingModel::SamplePart(int32 PartId, int32 Index)
+{
+	FNearestNeighborGeomCacheSampler* Sampler = static_cast<FNearestNeighborGeomCacheSampler*>(EditorModel->GetSampler());
+	Sampler->SamplePart(Index, NearestNeighborModel->PartVertexMap(PartId));
+	PartSampleDeltas = Sampler->GetPartVertexDeltas();
+	SampleBoneRotations = Sampler->GetBoneRotations();
+}
+
+void UNearestNeighborTrainingModel::SetSamplerPartData(const int32 PartId)
+{
+	GetNearestNeighborEditorModel()->SetSamplerPartData(PartId);
+}
+
+int32 UNearestNeighborTrainingModel::GetPartNumNeighbors(const int32 PartId) const
+{
+	return NearestNeighborModel->GetNumNeighborsFromGeometryCache(PartId);
+}
+
+void UNearestNeighborTrainingModel::SampleKmeansAnim(const int32 SkeletonId)
+{
+	FNearestNeighborGeomCacheSampler* Sampler = static_cast<FNearestNeighborGeomCacheSampler*>(EditorModel->GetSampler());
+	Sampler->SampleKMeansAnim(SkeletonId);
+}
+
+void UNearestNeighborTrainingModel::SampleKmeansFrame(const int32 Frame)
+{
+	FNearestNeighborGeomCacheSampler* Sampler = static_cast<FNearestNeighborGeomCacheSampler*>(EditorModel->GetSampler());
+	Sampler->SampleKMeansFrame(Frame);
+	SampleBoneRotations = Sampler->GetBoneRotations();
+}
+
+
+int32 UNearestNeighborTrainingModel::GetKmeansNumAnims() const
+{
+	return NearestNeighborModel->SourceSkeletons.Num();
+}
+
+int32 UNearestNeighborTrainingModel::GetKmeansAnimNumFrames(const int32 SkeletonId) const
+{
+	if (SkeletonId < NearestNeighborModel->SourceSkeletons.Num())
+	{
+		return NearestNeighborModel->SourceSkeletons[SkeletonId]->GetDataModel()->GetNumberOfFrames();
+	}
+	return 0;
+}
+
+int32 UNearestNeighborTrainingModel::GetKmeansNumClusters() const
+{
+	return NearestNeighborModel->NumClusters;
+}
+
+
+#undef LOCTEXT_NAMESPACE
diff --git a/Engine/Plugins/Animation/MLDeformer/NearestNeighborModel/Source/NearestNeighborModelEditor/Public/NearestNeighborEditorModel.h b/Engine/Plugins/Animation/MLDeformer/NearestNeighborModel/Source/NearestNeighborModelEditor/Public/NearestNeighborEditorModel.h
new file mode 100644
--- /dev/null
+++ b/Engine/Plugins/Animation/MLDeformer/NearestNeighborModel/Source/NearestNeighborModelEditor/Public/NearestNeighborEditorModel.h
@@ -0,0 +1,90 @@
+// Copyright Epic Games, Inc. All Rights Reserved.
+
+#pragma once
+
+#include "CoreMinimal.h"
+#include "MLDeformerMorphModelEditorModel.h"
+#include "MLDeformerEditorActor.h"
+
+class UMLDeformerAsset;
+class USkeletalMesh;
+class UGeometryCache;
+class UNearestNeighborModel;
+class UNearestNeighborModelVizSettings;
+
+namespace UE::NearestNeighborModel
+{
+	using namespace UE::MLDeformer;
+
+	class FNearestNeighborEditorModelActor;
+	class FNearestNeighborModelSampler;
+
+	class NEARESTNEIGHBORMODELEDITOR_API FNearestNeighborEditorModel 
+		: public UE::MLDeformer::FMLDeformerMorphModelEditorModel
+	{
+	public:
+		// We need to implement this static MakeInstance method.
+		static FMLDeformerEditorModel* MakeInstance();
+
+		// FGCObject overrides.
+		virtual FString GetReferencerName() const override { return TEXT("FNearestNeighborEditorModel"); }
+		// ~END FGCObject overrides.
+	
+		// FMLDeformerGeomCacheEditorModel overrides.
+		virtual void Init(const InitSettings& Settings) override;
+		virtual FMLDeformerSampler* CreateSampler() const override;
+		virtual void CreateActors(const TSharedRef<IPersonaPreviewScene>& InPersonaPreviewScene) override;
+		virtual void Tick(FEditorViewportClient* ViewportClient, float DeltaTime) override;
+		virtual FMLDeformerEditorActor* CreateEditorActor(const FMLDeformerEditorActor::FConstructSettings& Settings) const override;
+
+		virtual void InitInputInfo(UMLDeformerInputInfo* InputInfo) override;
+		virtual ETrainingResult Train() override;
+		virtual void OnPropertyChanged(FPropertyChangedEvent& PropertyChangedEvent) override;
+		virtual void OnPostTraining(ETrainingResult TrainingResult, bool bUsePartiallyTrainedWhenAborted) override;
+		// ~END FMLDeformerGeomCacheEditorModel overrides.
+
+		// Some helpers that cast to this model's variants of some classes.
+		UNearestNeighborModel* GetNearestNeighborModel() const { return static_cast<UNearestNeighborModel*>(Model); }
+	
+		// Recomputes nearest neighbor coeffcients and nearest neighbor vertex offsets. 
+		virtual void UpdateNearestNeighborData();
+
+		// Temporarily set sampler to use part data in nearest neighbor dataset. This is useful when updating nearest neighbor data.
+		bool SetSamplerPartData(const int32 PartI);
+
+		// Reset smapler to use the original dataset.
+		void ResetSamplerData();
+
+		void UpdateNearestNeighborActors();
+		void KMeansClusterPoses();
+
+		void InitMorphTargets();
+		void AddFloatArrayToDeltaArray(const TArray<float>& FloatArr, const TArray<uint32>& VertexMap, TArray<FVector3f>& DeltaArr, int32 DeltaArrayOffset = -1, float ScaleFactor = 1);
+
+		int32 GetNumParts();
+
+	protected:
+		virtual void CreateNearestNeighborActors(UWorld* World, int32 StartIndex = 0);
+
+		template<class TrainingModelClass>
+		TrainingModelClass* InitTrainingModel(FMLDeformerEditorModel* EditorModel)
+		{
+			TArray<UClass*> TrainingModels;
+			GetDerivedClasses(TrainingModelClass::StaticClass(), TrainingModels);
+			if (TrainingModels.IsEmpty())
+			{
+				return nullptr;
+			}
+
+			TrainingModelClass* TrainingModel = Cast<TrainingModelClass>(TrainingModels.Last()->GetDefaultObject());
+			TrainingModel->Init(EditorModel);
+			return TrainingModel;
+		}
+
+		// Actors showing the nearest neighbors of the current frame
+		TArray<FNearestNeighborEditorModelActor*> NearestNeighborActors;
+
+	private:
+		UWorld* EditorWorld = nullptr;
+	};
+}	// namespace UE::NearestNeighborModel
diff --git a/Engine/Plugins/Animation/MLDeformer/NearestNeighborModel/Source/NearestNeighborModelEditor/Public/NearestNeighborEditorModelActor.h b/Engine/Plugins/Animation/MLDeformer/NearestNeighborModel/Source/NearestNeighborModelEditor/Public/NearestNeighborEditorModelActor.h
new file mode 100644
--- /dev/null
+++ b/Engine/Plugins/Animation/MLDeformer/NearestNeighborModel/Source/NearestNeighborModelEditor/Public/NearestNeighborEditorModelActor.h
@@ -0,0 +1,37 @@
+// Copyright Epic Games, Inc. All Rights Reserved.
+
+#pragma once
+
+#include "MLDeformerGeomCacheActor.h"
+#include "NearestNeighborModel.h"
+
+class UMLDeformerComponent;
+class UGeometryCacheComponent;
+
+namespace UE::NearestNeighborModel
+{
+	using namespace UE::MLDeformer;
+
+	enum : int32
+	{
+		ActorID_NearestNeighborActors = 6
+	};
+
+	class NEARESTNEIGHBORMODELEDITOR_API FNearestNeighborEditorModelActor
+		: public FMLDeformerGeomCacheActor
+	{
+	public:
+		FNearestNeighborEditorModelActor(const FConstructSettings& Settings);
+
+		void SetGeometryCacheComponent(UGeometryCacheComponent* Component) { GeomCacheComponent = Component; }
+		UGeometryCacheComponent* GetGeometryCacheComponent() const { return GeomCacheComponent; }
+
+		void InitNearestNeighborActor(UNearestNeighborModel* InNearestNeighborModel, const int32 InPartId);
+		void TickNearestNeighborActor();
+
+	protected:
+		/** The geometry cache component (can be nullptr). */
+		TObjectPtr<UNearestNeighborModel> NearestNeighborModel = nullptr;
+		int32 PartId = INDEX_NONE;
+	};
+}	// namespace UE::NearestNeighborModel
diff --git a/Engine/Plugins/Animation/MLDeformer/NearestNeighborModel/Source/NearestNeighborModelEditor/Public/NearestNeighborGeomCacheSampler.h b/Engine/Plugins/Animation/MLDeformer/NearestNeighborModel/Source/NearestNeighborModelEditor/Public/NearestNeighborGeomCacheSampler.h
new file mode 100644
--- /dev/null
+++ b/Engine/Plugins/Animation/MLDeformer/NearestNeighborModel/Source/NearestNeighborModelEditor/Public/NearestNeighborGeomCacheSampler.h
@@ -0,0 +1,21 @@
+// Copyright Epic Games, Inc. All Rights Reserved.
+
+#pragma once
+
+#include "MLDeformerGeomCacheSampler.h"
+
+namespace UE::NearestNeighborModel
+{
+	class NEARESTNEIGHBORMODELEDITOR_API FNearestNeighborGeomCacheSampler
+		: public UE::MLDeformer::FMLDeformerGeomCacheSampler
+	{
+	public:
+		virtual void SamplePart(int32 InAnimFrameIndex, const TArray<uint32>& VertexMap);
+		virtual void SampleKMeansAnim(const int32 SkeletonId);
+		virtual void SampleKMeansFrame(const int32 Frame);
+		const TArray<float>& GetPartVertexDeltas() const { return PartVertexDeltas; }
+		void GeneratePartMeshMappings(const TArray<uint32>& VertexMap);
+	protected:
+		TArray<float> PartVertexDeltas;
+	};
+}
diff --git a/Engine/Plugins/Animation/MLDeformer/NearestNeighborModel/Source/NearestNeighborModelEditor/Public/NearestNeighborModelDetails.h b/Engine/Plugins/Animation/MLDeformer/NearestNeighborModel/Source/NearestNeighborModelEditor/Public/NearestNeighborModelDetails.h
new file mode 100644
--- /dev/null
+++ b/Engine/Plugins/Animation/MLDeformer/NearestNeighborModel/Source/NearestNeighborModelEditor/Public/NearestNeighborModelDetails.h
@@ -0,0 +1,37 @@
+// Copyright Epic Games, Inc. All Rights Reserved.
+
+#pragma once
+
+#include "CoreMinimal.h"
+#include "MLDeformerMorphModelDetails.h"
+
+class UNearestNeighborModel;
+
+namespace UE::NearestNeighborModel
+{
+	class FNearestNeighborEditorModel;
+
+	class NEARESTNEIGHBORMODELEDITOR_API FNearestNeighborModelDetails
+		: public UE::MLDeformer::FMLDeformerMorphModelDetails
+	{
+	public:
+		/** Makes a new instance of this detail layout class for a specific detail view requesting it. */
+		static TSharedRef<IDetailCustomization> MakeInstance();
+
+		// ILayoutDetails overrides.
+		virtual void CustomizeDetails(IDetailLayoutBuilder& DetailBuilder) override;
+		// ~END ILayoutDetails overrides.
+
+		virtual void CreateCategories() override;
+
+		virtual bool UpdateMemberPointers(const TArray<TWeakObjectPtr<UObject>>& Objects) override;
+
+	protected:
+		UNearestNeighborModel* NearestNeighborModel = nullptr;
+		FNearestNeighborEditorModel* NearestNeighborEditorModel = nullptr;
+
+		IDetailCategoryBuilder* ClothPartCategoryBuilder = nullptr;
+		IDetailCategoryBuilder* NearestNeighborCategoryBuilder = nullptr;
+		IDetailCategoryBuilder* KMeansCategoryBuilder = nullptr;
+	};
+}	// namespace UE::NearestNeighborModel
diff --git a/Engine/Plugins/Animation/MLDeformer/NearestNeighborModel/Source/NearestNeighborModelEditor/Public/NearestNeighborModelStyle.h b/Engine/Plugins/Animation/MLDeformer/NearestNeighborModel/Source/NearestNeighborModelEditor/Public/NearestNeighborModelStyle.h
new file mode 100644
--- /dev/null
+++ b/Engine/Plugins/Animation/MLDeformer/NearestNeighborModel/Source/NearestNeighborModelEditor/Public/NearestNeighborModelStyle.h
@@ -0,0 +1,19 @@
+// Copyright Epic Games, Inc. All Rights Reserved.
+
+#pragma once
+
+#include "Styling/SlateStyle.h"
+
+namespace UE::NearestNeighborModel
+{
+	class NEARESTNEIGHBORMODELEDITOR_API FNearestNeighborModelEditorStyle
+		: public FSlateStyleSet
+	{
+	public:
+		FNearestNeighborModelEditorStyle();
+		~FNearestNeighborModelEditorStyle();
+
+		static FNearestNeighborModelEditorStyle& Get();
+	};
+
+}	// namespace UE::NearestNeighborModel
diff --git a/Engine/Plugins/Animation/MLDeformer/NearestNeighborModel/Source/NearestNeighborModelEditor/Public/NearestNeighborModelVizSettingsDetails.h b/Engine/Plugins/Animation/MLDeformer/NearestNeighborModel/Source/NearestNeighborModelEditor/Public/NearestNeighborModelVizSettingsDetails.h
new file mode 100644
--- /dev/null
+++ b/Engine/Plugins/Animation/MLDeformer/NearestNeighborModel/Source/NearestNeighborModelEditor/Public/NearestNeighborModelVizSettingsDetails.h
@@ -0,0 +1,20 @@
+// Copyright Epic Games, Inc. All Rights Reserved.
+
+#pragma once
+
+#include "CoreMinimal.h"
+#include "MLDeformerMorphModelVizSettingsDetails.h"
+
+namespace UE::NearestNeighborModel
+{
+	class NEARESTNEIGHBORMODELEDITOR_API FNearestNeighborModelVizSettingsDetails
+		: public UE::MLDeformer::FMLDeformerMorphModelVizSettingsDetails
+	{
+	public:
+		/** Makes a new instance of this detail layout class for a specific detail view requesting it. */
+		static TSharedRef<IDetailCustomization> MakeInstance()
+		{
+			return MakeShareable(new FNearestNeighborModelVizSettingsDetails());
+		}
+	};
+}	// namespace UE::NearestNeighborModel
diff --git a/Engine/Plugins/Animation/MLDeformer/NearestNeighborModel/Source/NearestNeighborModelEditor/Public/NearestNeighborTrainingModel.h b/Engine/Plugins/Animation/MLDeformer/NearestNeighborModel/Source/NearestNeighborModelEditor/Public/NearestNeighborTrainingModel.h
new file mode 100644
--- /dev/null
+++ b/Engine/Plugins/Animation/MLDeformer/NearestNeighborModel/Source/NearestNeighborModelEditor/Public/NearestNeighborTrainingModel.h
@@ -0,0 +1,73 @@
+// Copyright Epic Games, Inc. All Rights Reserved.
+
+#pragma once
+
+#include "CoreMinimal.h"
+#include "MLDeformerTrainingModel.h"
+#include "NearestNeighborTrainingModel.generated.h"
+
+namespace UE::NearestNeighborModel
+{
+	class FNearestNeighborEditorModel;
+}
+
+class UNearestNeighborModel;
+
+UCLASS(Blueprintable)
+class NEARESTNEIGHBORMODELEDITOR_API UNearestNeighborTrainingModel
+	: public UMLDeformerTrainingModel
+{
+	GENERATED_BODY()
+
+public:
+	virtual void Init(UE::MLDeformer::FMLDeformerEditorModel* InEditorModel) override;
+
+	/** Main training function, with implementation in python. */
+	UFUNCTION(BlueprintImplementableEvent, Category = "Python")
+	int32 Train() const;
+
+	UFUNCTION(BlueprintPure, Category = "Python")
+	UNearestNeighborModel* GetNearestNeighborModel() const;
+
+	UE::NearestNeighborModel::FNearestNeighborEditorModel* GetNearestNeighborEditorModel() const;
+
+	UFUNCTION(BlueprintPure, Category = "Python")
+	const TArray<int32> GetPartVertexMap(const int32 PartId) const;
+
+	UFUNCTION(BlueprintCallable, Category = "Python")
+	void SamplePart(int32 PartId, int32 Index);
+
+	UFUNCTION(BlueprintImplementableEvent, Category = "Python")
+	void UpdateNearestNeighborData();
+
+	UFUNCTION(BlueprintCallable, Category = "Python")
+	void SetSamplerPartData(const int32 PartId);
+
+	UFUNCTION(BlueprintPure, Category = "Python")
+	int32 GetPartNumNeighbors(const int32 PartId) const;
+
+	UFUNCTION(BlueprintImplementableEvent, Category = "Python")
+	void KmeansClusterPoses() const;
+
+	UFUNCTION(BlueprintCallable, Category = "Python")
+	void SampleKmeansAnim(const int32 SkeletonId);
+
+	UFUNCTION(BlueprintCallable, Category = "Python")
+	void SampleKmeansFrame(const int32 Frame);
+
+	UFUNCTION(BlueprintPure, Category = "Python")
+	int32 GetKmeansNumAnims() const;
+
+	UFUNCTION(BlueprintPure, Category = "Python")
+	int32 GetKmeansAnimNumFrames(const int32 SkeletonId) const;
+
+	UFUNCTION(BlueprintPure, Category = "Python")
+	int32 GetKmeansNumClusters() const;
+
+
+public:
+	UPROPERTY(EditAnywhere, BlueprintReadWrite, Category = "Training Data")
+	TArray<float> PartSampleDeltas;
+
+	UNearestNeighborModel* NearestNeighborModel = nullptr;
+};
diff --git a/Engine/Plugins/Experimental/NNI/Source/NeuralNetworkInference/Private/NeuralNetwork.cpp b/Engine/Plugins/Experimental/NNI/Source/NeuralNetworkInference/Private/NeuralNetwork.cpp
--- a/Engine/Plugins/Experimental/NNI/Source/NeuralNetworkInference/Private/NeuralNetwork.cpp
+++ b/Engine/Plugins/Experimental/NNI/Source/NeuralNetworkInference/Private/NeuralNetwork.cpp
@@ -99,6 +99,17 @@ bool UNeuralNetwork::Load(const FString& InModelFilePath)
 	return Load();
 }
 
+bool UNeuralNetwork::Save(const FString& OutModelFilePath) const
+{
+	const FString OutPath = IFileManager::Get().ConvertToAbsolutePathForExternalAppForRead(*OutModelFilePath);
+	if (!FFileHelper::SaveArrayToFile(ModelReadFromFileInBytes, *OutPath))
+	{
+		UE_LOG(LogNeuralNetworkInference, Warning, TEXT("UNeuralNetwork::Save(): Error saving model \"%s\"."), *OutPath);
+		return false;
+	}
+	return true;
+}
+
 
 bool UNeuralNetwork::Load(TArray<uint8>& InModelReadFromFileInBytes)
 {
diff --git a/Engine/Plugins/Experimental/NNI/Source/NeuralNetworkInference/Public/NeuralNetwork.h b/Engine/Plugins/Experimental/NNI/Source/NeuralNetworkInference/Public/NeuralNetwork.h
--- a/Engine/Plugins/Experimental/NNI/Source/NeuralNetworkInference/Public/NeuralNetwork.h
+++ b/Engine/Plugins/Experimental/NNI/Source/NeuralNetworkInference/Public/NeuralNetwork.h
@@ -117,6 +117,13 @@ public:
 	 */
 	bool IsLoaded() const;
 
+
+	/**
+	 *  Save ONNX network to disk. 
+	 * @return whether the network was successfully saved.
+	 */
+	bool Save(const FString& OutModelFilePath) const;
+
 	/**
 	 * Getter and setter functions for DeviceType, InputDeviceType, and OutputDeviceType:
 	 * - GetDeviceType() returns DeviceType.
