Fix static analysis warning (Issue 213544)

diff --git a/Engine/Plugins/Animation/MLDeformer/NearestNeighborModel/Source/NearestNeighborModelEditor/Private/NearestNeighborModelDetails.cpp b/Engine/Plugins/Animation/MLDeformer/NearestNeighborModel/Source/NearestNeighborModelEditor/Private/NearestNeighborModelDetails.cpp
--- a/Engine/Plugins/Animation/MLDeformer/NearestNeighborModel/Source/NearestNeighborModelEditor/Private/NearestNeighborModelDetails.cpp
+++ b/Engine/Plugins/Animation/MLDeformer/NearestNeighborModel/Source/NearestNeighborModelEditor/Private/NearestNeighborModelDetails.cpp
@@ -58,7 +58,7 @@ namespace UE::NearestNeighborModel
 
 		// Cloth part settings
 		ClothPartCategoryBuilder->AddProperty(UNearestNeighborModel::GetClothPartEditorDataPropertyName());
-		FText ButtonText = NearestNeighborModel->IsClothPartDataValid() ? LOCTEXT("Update", "Update") : LOCTEXT("Update *", "Update *");
+		FText ButtonText = (NearestNeighborModel && NearestNeighborModel->IsClothPartDataValid()) ? LOCTEXT("Update", "Update") : LOCTEXT("Update *", "Update *");
 		ClothPartCategoryBuilder->AddCustomRow(FText::FromString(""))
 			.WholeRowContent()
 			[
@@ -83,7 +83,7 @@ namespace UE::NearestNeighborModel
 		// Nearest Neighbor settings
 		NearestNeighborCategoryBuilder->AddProperty(GET_MEMBER_NAME_STRING_CHECKED(UNearestNeighborModel, DecayFactor));
 		NearestNeighborCategoryBuilder->AddProperty(UNearestNeighborModel::GetNearestNeighborDataPropertyName());
-		ButtonText = NearestNeighborModel->IsNearestNeighborDataValid() ? LOCTEXT("Update", "Update") : LOCTEXT("Update *", "Update *");
+		ButtonText = (NearestNeighborModel && NearestNeighborModel->IsNearestNeighborDataValid()) ? LOCTEXT("Update", "Update") : LOCTEXT("Update *", "Update *");
 
 		NearestNeighborCategoryBuilder->AddCustomRow(FText::FromString(""))
 			.WholeRowContent()
