Add guard to nearestneighbormodel.py for import statement.
Provide a numpy PCA implementation when sklearn is not available.
Unwrap MACROs in NearestNeighborModel because unreal header tool does not work with MACRO fucntions.
Add warnings in nearest neighbor model UI.

diff --git a/Engine/Plugins/Animation/MLDeformer/NearestNeighborModel/Source/NearestNeighborModel/Public/NearestNeighborModel.h b/Engine/Plugins/Animation/MLDeformer/NearestNeighborModel/Source/NearestNeighborModel/Public/NearestNeighborModel.h
--- a/Engine/Plugins/Animation/MLDeformer/NearestNeighborModel/Source/NearestNeighborModel/Public/NearestNeighborModel.h
+++ b/Engine/Plugins/Animation/MLDeformer/NearestNeighborModel/Source/NearestNeighborModel/Public/NearestNeighborModel.h
@@ -202,12 +202,6 @@ public:
 	void InitInputInfo();
 	void InitPreviousWeights();
 
-#define EDITORONLY_PROPERTY_HELPERS(TYPE, NAME) \
-	UFUNCTION(BlueprintPure, Category = "Nearest Neighbor Model")\
-	const TYPE Get##NAME() const { return NAME; }\
-	\
-	static FName Get##NAME##PropertyName() { return GET_MEMBER_NAME_CHECKED(UNearestNeighborModel, NAME); }
-
 #if WITH_EDITORONLY_DATA
 	TObjectPtr<UAnimSequence> GetNearestNeighborSkeletons(int32 PartId);
 	const TObjectPtr<UAnimSequence> GetNearestNeighborSkeletons(int32 PartId) const;
@@ -229,16 +223,40 @@ public:
 	void ValidateNearestNeighborData() {bNearestNeighborDataValid = true; }
 	bool IsClothPartDataValid() { return bClothPartDataValid; }
 	bool IsNearestNeighborDataValid() { return bNearestNeighborDataValid; }
+#endif
+
+#if WITH_EDITOR
+	UFUNCTION(BlueprintPure, Category = "Nearest Neighbor Model")
+	const int32	GetInputDim() const { return InputDim; }
+	static FName GetInputDimPropertyName() { return GET_MEMBER_NAME_CHECKED(UNearestNeighborModel, InputDim); }
+
+	UFUNCTION(BlueprintPure, Category = "Nearest Neighbor Model")
+	const TArray<int32>& GetHiddenLayerDims() const { return HiddenLayerDims; }
+	static FName GetHiddenLayerDimsPropertyName() { return GET_MEMBER_NAME_CHECKED(UNearestNeighborModel, HiddenLayerDims); }
+
+	UFUNCTION(BlueprintPure, Category = "Nearest Neighbor Model")
+	const int32	GetOutputDim() const { return OutputDim; }
+	static FName GetOutputDimPropertyName() { return GET_MEMBER_NAME_CHECKED(UNearestNeighborModel, OutputDim); }
 
+	UFUNCTION(BlueprintPure, Category = "Nearest Neighbor Model")
+	const int32	GetNumEpochs() const { return NumEpochs; }
+	static FName GetNumEpochsPropertyName() { return GET_MEMBER_NAME_CHECKED(UNearestNeighborModel, NumEpochs); }
+
+	UFUNCTION(BlueprintPure, Category = "Nearest Neighbor Model")
+	const int32	GetBatchSize() const { return BatchSize; }
+	static FName GetBatchSizePropertyName() { return GET_MEMBER_NAME_CHECKED(UNearestNeighborModel, BatchSize); }
+
+	UFUNCTION(BlueprintPure, Category = "Nearest Neighbor Model")
+	const int32	GetLearningRate() const { return LearningRate; }
+	static FName GetLearningRatePropertyName() { return GET_MEMBER_NAME_CHECKED(UNearestNeighborModel, LearningRate); }
+
+	UFUNCTION(BlueprintPure, Category = "Nearest Neighbor Model")
+	const TArray<FClothPartEditorData>&	GetClothPartEditorData() const { return ClothPartEditorData; }
+	static FName GetClothPartEditorDataPropertyName() { return GET_MEMBER_NAME_CHECKED(UNearestNeighborModel, ClothPartEditorData); }
 
-	EDITORONLY_PROPERTY_HELPERS(int32, InputDim)
-	EDITORONLY_PROPERTY_HELPERS(TArray<int32>, HiddenLayerDims)
-	EDITORONLY_PROPERTY_HELPERS(int32, OutputDim)
-	EDITORONLY_PROPERTY_HELPERS(int32, NumEpochs)
-	EDITORONLY_PROPERTY_HELPERS(int32, BatchSize)
-	EDITORONLY_PROPERTY_HELPERS(float, LearningRate)
-	EDITORONLY_PROPERTY_HELPERS(TArray<FClothPartEditorData>, ClothPartEditorData)
-	EDITORONLY_PROPERTY_HELPERS(TArray<FSkeletonCachePair>, NearestNeighborData)
+	UFUNCTION(BlueprintPure, Category = "Nearest Neighbor Model")
+	const TArray<FSkeletonCachePair>& GetNearestNeighborData() const { return NearestNeighborData; }
+	static FName GetNearestNeighborDataPropertyName() { return GET_MEMBER_NAME_CHECKED(UNearestNeighborModel, NearestNeighborData); }
 #endif
 
 
@@ -281,10 +299,10 @@ public:
 	UPROPERTY()
 	TArray<FClothPartData> ClothPartData;
 
-	UPROPERTY(EditAnywhere, BlueprintReadWrite, Category = "Network Inputs")
+	UPROPERTY(BlueprintReadWrite, Category = "Network Inputs")
 	TArray<float> InputsMin;
 
-	UPROPERTY(EditAnywhere, BlueprintReadWrite, Category = "Network Inputs")
+	UPROPERTY(BlueprintReadWrite, Category = "Network Inputs")
 	TArray<float> InputsMax;
 
 	UPROPERTY(EditAnywhere, BlueprintReadWrite, Category = "KMeans Pose Generator")
diff --git a/Engine/Plugins/Animation/MLDeformer/NearestNeighborModel/Source/NearestNeighborModelEditor/Private/NearestNeighborModelDetails.cpp b/Engine/Plugins/Animation/MLDeformer/NearestNeighborModel/Source/NearestNeighborModelEditor/Private/NearestNeighborModelDetails.cpp
--- a/Engine/Plugins/Animation/MLDeformer/NearestNeighborModel/Source/NearestNeighborModelEditor/Private/NearestNeighborModelDetails.cpp
+++ b/Engine/Plugins/Animation/MLDeformer/NearestNeighborModel/Source/NearestNeighborModelEditor/Private/NearestNeighborModelDetails.cpp
@@ -7,6 +7,7 @@
 #include "DetailLayoutBuilder.h"
 #include "DetailCategoryBuilder.h"
 #include "DetailWidgetRow.h"
+#include "SWarningOrErrorBox.h"
 #include "IDetailsView.h"
 #include "Widgets/Input/SButton.h"
 
@@ -41,6 +42,19 @@ namespace UE::NearestNeighborModel
 		ClothPartCategoryBuilder = &DetailLayoutBuilder->EditCategory("Cloth Parts", FText::GetEmpty(), ECategoryPriority::Important);
 		NearestNeighborCategoryBuilder = &DetailLayoutBuilder->EditCategory("Nearest Neighbors", FText::GetEmpty(), ECategoryPriority::Important);
 		KMeansCategoryBuilder = &DetailLayoutBuilder->EditCategory("KMeans Pose Generator", FText::GetEmpty(), ECategoryPriority::Important);
+
+		// Add warning in CreateCategories so that the warning appears at the top of the details panel.
+		FDetailWidgetRow& NearestNeighborWarningRow = BaseMeshCategoryBuilder->AddCustomRow(FText::FromString("NearestNeighborWarning"))
+			.WholeRowContent()
+			[
+				SNew(SBox)
+				.Padding(FMargin(0.0f, 4.0f))
+				[
+					SNew(SWarningOrErrorBox)
+					.MessageStyle(EMessageStyle::Warning)
+					.Message(LOCTEXT("NearestNeighborWarning", "Nearest neighbor model is still experimental and the details here are subject to change."))
+				]
+			];
 	}
 
 	void FNearestNeighborModelDetails::CustomizeDetails(IDetailLayoutBuilder& DetailBuilder)
@@ -83,7 +97,7 @@ namespace UE::NearestNeighborModel
 		// Nearest Neighbor settings
 		NearestNeighborCategoryBuilder->AddProperty(GET_MEMBER_NAME_STRING_CHECKED(UNearestNeighborModel, DecayFactor));
 		NearestNeighborCategoryBuilder->AddProperty(UNearestNeighborModel::GetNearestNeighborDataPropertyName());
-		ButtonText = (NearestNeighborModel && NearestNeighborModel->IsNearestNeighborDataValid()) ? LOCTEXT("Update", "Update") : LOCTEXT("Update *", "Update *");
+		ButtonText = NearestNeighborModel->IsNearestNeighborDataValid() ? LOCTEXT("Update", "Update") : LOCTEXT("Update *", "Update *");
 
 		NearestNeighborCategoryBuilder->AddCustomRow(FText::FromString(""))
 			.WholeRowContent()
