Change learning rate return type to float

diff --git a/Engine/Plugins/Animation/MLDeformer/NearestNeighborModel/Source/NearestNeighborModel/Public/NearestNeighborModel.h b/Engine/Plugins/Animation/MLDeformer/NearestNeighborModel/Source/NearestNeighborModel/Public/NearestNeighborModel.h
--- a/Engine/Plugins/Animation/MLDeformer/NearestNeighborModel/Source/NearestNeighborModel/Public/NearestNeighborModel.h
+++ b/Engine/Plugins/Animation/MLDeformer/NearestNeighborModel/Source/NearestNeighborModel/Public/NearestNeighborModel.h
@@ -247,7 +247,7 @@ public:
 	static FName GetBatchSizePropertyName() { return GET_MEMBER_NAME_CHECKED(UNearestNeighborModel, BatchSize); }
 
 	UFUNCTION(BlueprintPure, Category = "Nearest Neighbor Model")
-	const int32	GetLearningRate() const { return LearningRate; }
+	const float	GetLearningRate() const { return LearningRate; }
 	static FName GetLearningRatePropertyName() { return GET_MEMBER_NAME_CHECKED(UNearestNeighborModel, LearningRate); }
 
 	UFUNCTION(BlueprintPure, Category = "Nearest Neighbor Model")
