Add NearestNeighborOffsetWeights to visualize the effect of nearest neighbor offsets
Add network size and morph data size in UI
Move NeareestNeighborIds to Instance so that PIE instance does not interfere with NearestNeighborEditor
Add refreshmorphtargets so that users can fine tune morph parameters in editor.

diff --git a/Engine/Plugins/Animation/MLDeformer/NearestNeighborModel/Source/NearestNeighborModel/Private/NearestNeighborModel.cpp b/Engine/Plugins/Animation/MLDeformer/NearestNeighborModel/Source/NearestNeighborModel/Private/NearestNeighborModel.cpp
--- a/Engine/Plugins/Animation/MLDeformer/NearestNeighborModel/Source/NearestNeighborModel/Private/NearestNeighborModel.cpp
+++ b/Engine/Plugins/Animation/MLDeformer/NearestNeighborModel/Source/NearestNeighborModel/Private/NearestNeighborModel.cpp
@@ -7,6 +7,7 @@
 #include "Engine/SkeletalMesh.h"
 #include "GeometryCache.h"
 #include "Misc/FileHelper.h"
+#include "Components/SkinnedMeshComponent.h"
 #include "Rendering/SkeletalMeshModel.h"
 #include "UObject/UObjectGlobals.h"
 #include "NeuralNetwork.h"
@@ -55,9 +56,10 @@ void UNearestNeighborModel::PostLoad()
 	InitPreviousWeights();
 
 #if WITH_EDITORONLY_DATA
-	NearestNeighborIds.SetNumZeroed(GetNumParts());
 	UpdateNetworkInputDim();
 	UpdateNetworkOutputDim();
+	UpdateNetworkSize();
+	UpdateMorphTargetSize();
 #endif
 }
 
@@ -155,7 +157,6 @@ void UNearestNeighborModel::UpdateClothPartData()
 	}
 	UpdatePCACoeffNums();
 
-	NearestNeighborIds.SetNumZeroed(GetNumParts());
 	NearestNeighborData.SetNumZeroed(GetNumParts());
 	UpdateNetworkInputDim();
 	UpdateNetworkOutputDim();
@@ -212,6 +213,40 @@ int32 UNearestNeighborModel::GetNumNeighborsFromGeometryCache(int32 PartId) cons
 		return 0;
 	}
 }
+
+void UNearestNeighborModel::UpdateNetworkSize()
+{
+	UNeuralNetwork* Network = GetNeuralNetwork();
+	if (Network != nullptr)
+	{
+		const SIZE_T NumBytes = Network->GetResourceSizeBytes(EResourceSizeMode::EstimatedTotal);
+		SavedNetworkSize = (double)NumBytes / 1024 / 1024;
+	}
+	else
+	{
+		SavedNetworkSize = 0.0f;
+	}
+}
+
+class FMorphTargetBuffersInfo : public FMorphTargetVertexInfoBuffers
+{
+public:
+	uint32 GetMorphDataSize() const { return MorphData.Num() * sizeof(uint32); }
+};
+
+void UNearestNeighborModel::UpdateMorphTargetSize()
+{
+	if (GetMorphTargetSet())
+	{
+		const FMorphTargetBuffersInfo* MorphBuffersInfo = static_cast<FMorphTargetBuffersInfo*>(&GetMorphTargetSet()->MorphBuffers);
+		const double Size = MorphBuffersInfo->GetMorphDataSize();
+		MorphDataSize = Size / 1024 / 1024;
+	}
+	else
+	{
+		MorphDataSize = 0.0f;
+	}
+}
 #endif
 
 void UNearestNeighborModel::InitInputInfo()
diff --git a/Engine/Plugins/Animation/MLDeformer/NearestNeighborModel/Source/NearestNeighborModel/Private/NearestNeighborModelInstance.cpp b/Engine/Plugins/Animation/MLDeformer/NearestNeighborModel/Source/NearestNeighborModel/Private/NearestNeighborModelInstance.cpp
--- a/Engine/Plugins/Animation/MLDeformer/NearestNeighborModel/Source/NearestNeighborModel/Private/NearestNeighborModelInstance.cpp
+++ b/Engine/Plugins/Animation/MLDeformer/NearestNeighborModel/Source/NearestNeighborModel/Private/NearestNeighborModelInstance.cpp
@@ -84,7 +84,7 @@ bool UNearestNeighborModelInstance::SetupInputs()
     UNeuralNetwork* NeuralNetwork = Model->GetNeuralNetwork();
     if (NeuralNetwork)
     {
-    	UNearestNeighborModel* NearestNeighborModel = static_cast<UNearestNeighborModel*>(Model.Get());
+		UNearestNeighborModel* NearestNeighborModel = static_cast<UNearestNeighborModel*>(Model.Get());
     	float* InputDataPointer = static_cast<float*>(NeuralNetwork->GetInputDataPointerMutableForContext(NeuralNetworkInferenceHandle));
     	const int64 NumNeuralNetInputs = NeuralNetwork->GetInputTensorForContext(NeuralNetworkInferenceHandle).Num();
     	NearestNeighborModel->ClipInputs(InputDataPointer, NumNeuralNetInputs);
@@ -119,6 +119,10 @@ void UNearestNeighborModelInstance::RunNearestNeighborModel(float ModelWeight)
 		return;
 	}
 
+#if WITH_EDITORONLY_DATA
+	NearestNeighborIds.SetNum(NearestNeighborModel->GetNumParts());
+#endif
+
 	const UNeuralNetwork* NeuralNetwork = Model->GetNeuralNetwork();
 	if (NeuralNetwork)
 	{
@@ -141,13 +145,13 @@ void UNearestNeighborModelInstance::RunNearestNeighborModel(float ModelWeight)
 			{
 				const int32 NearestNeighborId = FindNearestNeighbor(OutputTensor, PartId);
 	#if WITH_EDITORONLY_DATA
-				NearestNeighborModel->SetNearestNeighborId(PartId, NearestNeighborId);
+				NearestNeighborIds[PartId] = NearestNeighborId;
 	#endif
 
 				const int32 NumNeighbors = NearestNeighborModel->GetNumNeighbors(PartId);
 				for (int32 NeighborId = 0; NeighborId < NumNeighbors; NeighborId++)
 				{
-					const float W = NeighborId == NearestNeighborId ? ModelWeight : 0;
+					const float W = NeighborId == NearestNeighborId ? ModelWeight * NearestNeighborModel->GetNearestNeighborOffsetWeight() : 0;
 					UpdateWeight(WeightData->Weights, NeighborOffset + NeighborId, W);
 				}
 				NeighborOffset += NumNeighbors;
diff --git a/Engine/Plugins/Animation/MLDeformer/NearestNeighborModel/Source/NearestNeighborModel/Public/NearestNeighborModel.h b/Engine/Plugins/Animation/MLDeformer/NearestNeighborModel/Source/NearestNeighborModel/Public/NearestNeighborModel.h
--- a/Engine/Plugins/Animation/MLDeformer/NearestNeighborModel/Source/NearestNeighborModel/Public/NearestNeighborModel.h
+++ b/Engine/Plugins/Animation/MLDeformer/NearestNeighborModel/Source/NearestNeighborModel/Public/NearestNeighborModel.h
@@ -159,6 +159,7 @@ public:
 	int32 GetNumNeighbors(int32 PartId) const { return ClothPartData[PartId].NumNeighbors; }
 
 	float GetDecayFactor() const { return DecayFactor; };
+	float GetNearestNeighborOffsetWeight() const { return NearestNeighborOffsetWeight; }
 
 	const TArray<uint32>& PartVertexMap(int32 PartId) const { return ClothPartData[PartId].VertexMap; }
 
@@ -178,6 +179,9 @@ public:
 	UFUNCTION(BlueprintCallable, Category = "Nearest Neighbor Model")
 	void SetVertexMean(int32 PartId, const TArray<float>& VertexMean) { ClothPartData[PartId].VertexMean = VertexMean; }
 
+	UFUNCTION(BlueprintCallable, Category = "Nearest Neighbor Model")
+	void SetNumNeighbors(int32 PartId, int32 InNumNeighbors) { ClothPartData[PartId].NumNeighbors = InNumNeighbors; }
+
 	UFUNCTION(BlueprintPure, Category = "Nearest Neighbor Model")
 	const TArray<float>& NeighborCoeffs(int32 PartId) const { return ClothPartData[PartId].NeighborCoeffs; }
 
@@ -209,13 +213,12 @@ public:
 	const TObjectPtr<UGeometryCache> GetNearestNeighborCache(int32 PartId) const;
 	int32 GetNumNeighborsFromGeometryCache(int32 PartId) const;
 
-	uint32 NearestNeighborId(int32 PartId) const { return NearestNeighborIds[PartId]; }
-	void SetNearestNeighborId(int32 PartId, int32 Id) { NearestNeighborIds[PartId] = Id; }
-
 	void UpdateNetworkInputDim();
 	void UpdateNetworkOutputDim();
 	void UpdateClothPartData();
 	void UpdatePCACoeffNums();
+	void UpdateNetworkSize();
+	void UpdateMorphTargetSize();
 
 	void InvalidateClothPartData() { bClothPartDataValid = false; bNearestNeighborDataValid = false; }
 	void ValidateClothPartData() { bClothPartDataValid = true; }
@@ -257,6 +260,11 @@ public:
 	UFUNCTION(BlueprintPure, Category = "Nearest Neighbor Model")
 	const TArray<FSkeletonCachePair>& GetNearestNeighborData() const { return NearestNeighborData; }
 	static FName GetNearestNeighborDataPropertyName() { return GET_MEMBER_NAME_CHECKED(UNearestNeighborModel, NearestNeighborData); }
+
+	static FName GetNearestNeighborOffsetWeightPropertyName() { return GET_MEMBER_NAME_CHECKED(UNearestNeighborModel, NearestNeighborOffsetWeight); }
+
+	static FName GetSavedNetworkSizePropertyName() { return GET_MEMBER_NAME_CHECKED(UNearestNeighborModel, SavedNetworkSize); }
+	static FName GetMorphDataSizePropertyName() { return GET_MEMBER_NAME_CHECKED(UNearestNeighborModel, MorphDataSize); }
 #endif
 
 
@@ -292,7 +300,11 @@ protected:
 	UPROPERTY(EditAnywhere, BlueprintReadWrite, Category = "Nearest Neighbors")
 	TArray<FSkeletonCachePair> NearestNeighborData;
 
-	TArray<uint32> NearestNeighborIds;
+	UPROPERTY(VisibleAnywhere, BlueprintReadOnly, Category = "Training Settings", meta = (ClampMin = "0.0"))
+	float SavedNetworkSize = 0.0f;
+
+	UPROPERTY(VisibleAnywhere, BlueprintReadOnly, Category = "Morph Targets", meta = (ClampMin = "0.0"))
+	float MorphDataSize = 0.0f;
 #endif
 
 public:
@@ -314,5 +326,8 @@ public:
 	UPROPERTY(EditAnywhere, BlueprintReadWrite, Category = "Nearest Neighbors", META = (ClampMin = "0", ClampMax = "1"))
 	float DecayFactor = 0.85f;
 
+	UPROPERTY(EditAnywhere, BlueprintReadWrite, Category = "Nearest Neighbors", META = (ClampMin = "0", ClampMax = "1"))
+	float NearestNeighborOffsetWeight = 1.0f;
+
 	TArray<float> PreviousWeights; 
 };
diff --git a/Engine/Plugins/Animation/MLDeformer/NearestNeighborModel/Source/NearestNeighborModel/Public/NearestNeighborModelInstance.h b/Engine/Plugins/Animation/MLDeformer/NearestNeighborModel/Source/NearestNeighborModel/Public/NearestNeighborModelInstance.h
--- a/Engine/Plugins/Animation/MLDeformer/NearestNeighborModel/Source/NearestNeighborModel/Public/NearestNeighborModelInstance.h
+++ b/Engine/Plugins/Animation/MLDeformer/NearestNeighborModel/Source/NearestNeighborModel/Public/NearestNeighborModelInstance.h
@@ -16,6 +16,11 @@ public:
     virtual void Execute(float ModelWeight) override;
     virtual bool SetupInputs() override;
 
+#if WITH_EDITORONLY_DATA
+    uint32 NearestNeighborId(int32 PartId) const { return NearestNeighborIds[PartId]; }
+    int32 NeighborIdNum() const { return NearestNeighborIds.Num(); }
+#endif
+
 protected:
     virtual int64 SetBoneTransforms(float* OutputBuffer, int64 OutputBufferSize, int64 StartIndex) override;
     // ~END UMLDeformerModelInstance overrides
@@ -24,4 +29,8 @@ private:
     void RunNearestNeighborModel(float ModelWeight);
     int32 FindNearestNeighbor(const FNeuralTensor& PCACoeffTensor, int32 PartId);
     void UpdateWeight(TArray<float>& MorphWeights, int32 Index, float W);
+
+#if WITH_EDITORONLY_DATA
+    TArray<uint32> NearestNeighborIds;
+#endif
 };
diff --git a/Engine/Plugins/Animation/MLDeformer/NearestNeighborModel/Source/NearestNeighborModelEditor/Private/NearestNeighborEditorModel.cpp b/Engine/Plugins/Animation/MLDeformer/NearestNeighborModel/Source/NearestNeighborModelEditor/Private/NearestNeighborEditorModel.cpp
--- a/Engine/Plugins/Animation/MLDeformer/NearestNeighborModel/Source/NearestNeighborModelEditor/Private/NearestNeighborEditorModel.cpp
+++ b/Engine/Plugins/Animation/MLDeformer/NearestNeighborModel/Source/NearestNeighborModelEditor/Private/NearestNeighborEditorModel.cpp
@@ -2,11 +2,13 @@
 
 #include "NearestNeighborEditorModel.h"
 #include "NearestNeighborModel.h"
+#include "NearestNeighborModelInstance.h"
 #include "NearestNeighborTrainingModel.h"
 #include "NearestNeighborModelInputInfo.h"
 #include "NearestNeighborEditorModelActor.h"
 #include "NearestNeighborModelStyle.h"
 #include "NearestNeighborGeomCacheSampler.h"
+#include "MLDeformerComponent.h"
 #include "MLDeformerEditorToolkit.h"
 #include "NeuralNetwork.h"
 #include "GeometryCache.h"
@@ -119,7 +121,23 @@ namespace UE::NearestNeighborModel
 				LOCTEXT("TestNearestNeighborLabelText", "Nearest Neigbors"),
 				false);
 			FNearestNeighborEditorModelActor* NearestNeighborActor = static_cast<FNearestNeighborEditorModelActor*>(EditorActors.Last());
-			NearestNeighborActor->InitNearestNeighborActor(NearestNeighborModel, PartId);
+
+			UNearestNeighborModelInstance* ModelInstance = nullptr;
+			const FMLDeformerEditorActor* EditorActor = FindEditorActor(ActorID_Test_MLDeformed);
+			if (EditorActor)
+			{
+				const AActor* Actor = EditorActor->GetActor();
+				if (Actor)
+				{
+					UMLDeformerComponent* Component = Actor->FindComponentByClass<UMLDeformerComponent>();
+					if (Component)
+					{
+						ModelInstance = static_cast<UNearestNeighborModelInstance*>(Component->GetModelInstance());
+					}
+				}
+			}
+
+			NearestNeighborActor->InitNearestNeighborActor(ModelInstance, PartId);
 			NearestNeighborActor->SetMeshOffsetFactor(2.0f);
 			NearestNeighborActors[PartId] = NearestNeighborActor;
 		}
@@ -182,7 +200,7 @@ namespace UE::NearestNeighborModel
 				}
 				else
 				{
-					UE_LOG(LogNearestNeighborModel, Error, TEXT("Part %d frame mismatch: AnimSequence has %d frames and GeometryCache has %d frames"), PartId, AnimSequence->GetDataModel()->GetNumberOfFrames(), NumFrames);
+					UE_LOG(LogNearestNeighborModel, Error, TEXT("Part %d frame mismatch: AnimSequence has %d frames and GeometryCache has %d frames"), PartId, AnimSequence->GetDataModel()->GetNumberOfKeys(), NumFrames);
 				}
 			}
 			else
@@ -378,6 +396,11 @@ namespace UE::NearestNeighborModel
 			AddFloatArrayToDeltaArray(NearestNeighborModel->ClothPartData[PartId].NeighborOffsets, VertexMap, Deltas);
 		}
 
+		if (Deltas.Num() == 0)
+		{
+			return;
+		}
+
 		const int32 LOD = 0;
 		TArray<UMorphTarget*> MorphTargets;
 		CreateEngineMorphTargets(MorphTargets, Deltas, FString("NNMorphTarget_"), LOD, NearestNeighborModel->GetMorphTargetDeltaThreshold());
@@ -386,6 +409,21 @@ namespace UE::NearestNeighborModel
 		CompressEngineMorphTargets(MorphBuffers, MorphTargets, LOD, NearestNeighborModel->GetMorphTargetErrorTolerance());
 		NearestNeighborModel->SetMorphTargetDeltas(Deltas);
 	}
+
+	void FNearestNeighborEditorModel::RefreshMorphTargets()
+	{
+		USkeletalMeshComponent* SkelMeshComponent = FindEditorActor(ActorID_Test_MLDeformed)->GetSkeletalMeshComponent() ;
+		if (SkelMeshComponent)
+		{
+			FMorphTargetVertexInfoBuffers& MorphBuffers = GetNearestNeighborModel()->GetMorphTargetSet()->MorphBuffers;
+			BeginReleaseResource(&MorphBuffers);
+			if (MorphBuffers.IsMorphCPUDataValid() && MorphBuffers.GetNumMorphs() > 0)
+			{
+				BeginInitResource(&MorphBuffers);
+			}
+			SkelMeshComponent->RefreshExternalMorphTargetWeights();
+		}
+	}
 }	// namespace UE::NearestNeighborModel
 
 #undef LOCTEXT_NAMESPACE
diff --git a/Engine/Plugins/Animation/MLDeformer/NearestNeighborModel/Source/NearestNeighborModelEditor/Private/NearestNeighborEditorModelActor.cpp b/Engine/Plugins/Animation/MLDeformer/NearestNeighborModel/Source/NearestNeighborModelEditor/Private/NearestNeighborEditorModelActor.cpp
--- a/Engine/Plugins/Animation/MLDeformer/NearestNeighborModel/Source/NearestNeighborModelEditor/Private/NearestNeighborEditorModelActor.cpp
+++ b/Engine/Plugins/Animation/MLDeformer/NearestNeighborModel/Source/NearestNeighborModelEditor/Private/NearestNeighborEditorModelActor.cpp
@@ -12,18 +12,18 @@ namespace UE::NearestNeighborModel
 	{
 	}
 
-	void FNearestNeighborEditorModelActor::InitNearestNeighborActor(UNearestNeighborModel* InNearestNeighborModel, const int32 InPartId)
+	void FNearestNeighborEditorModelActor::InitNearestNeighborActor(UNearestNeighborModelInstance* InModelInstance, const int32 InPartId)
 	{
-		NearestNeighborModel = InNearestNeighborModel;
+		ModelInstance = InModelInstance;
 		PartId = InPartId;
 	}
 
 	void FNearestNeighborEditorModelActor::TickNearestNeighborActor()
 	{
-		if (GeomCacheComponent && GeomCacheComponent->GetGeometryCache() && NearestNeighborModel)
+		if (GeomCacheComponent && GeomCacheComponent->GetGeometryCache() && ModelInstance && PartId < ModelInstance->NeighborIdNum())
 		{
 			GeomCacheComponent->SetManualTick(true);
-			GeomCacheComponent->TickAtThisTime(GeomCacheComponent->GetTimeAtFrame(NearestNeighborModel->NearestNeighborId(PartId)), false, false, false);
+			GeomCacheComponent->TickAtThisTime(GeomCacheComponent->GetTimeAtFrame(ModelInstance->NearestNeighborId(PartId)), false, false, false);
 		}
 	}
 }	// namespace UE::NearestNeighborModel
diff --git a/Engine/Plugins/Animation/MLDeformer/NearestNeighborModel/Source/NearestNeighborModelEditor/Private/NearestNeighborGeomCacheSampler.cpp b/Engine/Plugins/Animation/MLDeformer/NearestNeighborModel/Source/NearestNeighborModelEditor/Private/NearestNeighborGeomCacheSampler.cpp
--- a/Engine/Plugins/Animation/MLDeformer/NearestNeighborModel/Source/NearestNeighborModelEditor/Private/NearestNeighborGeomCacheSampler.cpp
+++ b/Engine/Plugins/Animation/MLDeformer/NearestNeighborModel/Source/NearestNeighborModelEditor/Private/NearestNeighborGeomCacheSampler.cpp
@@ -113,6 +113,16 @@ namespace UE::NearestNeighborModel
 			}
 
 			const int32 NumPartVerts = VertexMap.Num();
+			if (GeomCacheMeshData.ImportedVertexNumbers.IsEmpty())
+			{
+				UE_LOG(LogNearestNeighborModel, Error, TEXT("Geometry cache has no stored vertex numbers"));
+				return;
+			}
+			const int32 MaxVertexNumber = FMath::Max(GeomCacheMeshData.ImportedVertexNumbers);
+			if (NumPartVerts != MaxVertexNumber + 1)
+			{
+				UE_LOG(LogNearestNeighborModel, Error, TEXT("Vertex number mismatch: part has %d vertices but geometry cache has %d vertices"), NumPartVerts, MaxVertexNumber + 1);
+			}
 
 			Mapping.SkelMeshToTrackVertexMap.AddUninitialized(NumPartVerts);
 			Mapping.ImportedVertexToRenderVertexMap.AddUninitialized(NumPartVerts);
diff --git a/Engine/Plugins/Animation/MLDeformer/NearestNeighborModel/Source/NearestNeighborModelEditor/Private/NearestNeighborModelDetails.cpp b/Engine/Plugins/Animation/MLDeformer/NearestNeighborModel/Source/NearestNeighborModelEditor/Private/NearestNeighborModelDetails.cpp
--- a/Engine/Plugins/Animation/MLDeformer/NearestNeighborModel/Source/NearestNeighborModelEditor/Private/NearestNeighborModelDetails.cpp
+++ b/Engine/Plugins/Animation/MLDeformer/NearestNeighborModel/Source/NearestNeighborModelEditor/Private/NearestNeighborModelDetails.cpp
@@ -69,6 +69,7 @@ namespace UE::NearestNeighborModel
 		TrainingSettingsCategoryBuilder->AddProperty(UNearestNeighborModel::GetNumEpochsPropertyName());
 		TrainingSettingsCategoryBuilder->AddProperty(UNearestNeighborModel::GetBatchSizePropertyName());
 		TrainingSettingsCategoryBuilder->AddProperty(UNearestNeighborModel::GetLearningRatePropertyName());
+		TrainingSettingsCategoryBuilder->AddProperty(UNearestNeighborModel::GetSavedNetworkSizePropertyName());
 
 		// Cloth part settings
 		ClothPartCategoryBuilder->AddProperty(UNearestNeighborModel::GetClothPartEditorDataPropertyName());
@@ -84,6 +85,7 @@ namespace UE::NearestNeighborModel
 					if (NearestNeighborModel != nullptr)
 					{
 						NearestNeighborModel->UpdateClothPartData();
+						NearestNeighborModel->InitPreviousWeights();
 						if (NearestNeighborEditorModel != nullptr)
 						{
 							NearestNeighborEditorModel->UpdateNearestNeighborActors();
@@ -96,6 +98,7 @@ namespace UE::NearestNeighborModel
 
 		// Nearest Neighbor settings
 		NearestNeighborCategoryBuilder->AddProperty(GET_MEMBER_NAME_STRING_CHECKED(UNearestNeighborModel, DecayFactor));
+		NearestNeighborCategoryBuilder->AddProperty(GET_MEMBER_NAME_STRING_CHECKED(UNearestNeighborModel, NearestNeighborOffsetWeight));
 		NearestNeighborCategoryBuilder->AddProperty(UNearestNeighborModel::GetNearestNeighborDataPropertyName());
 		ButtonText = NearestNeighborModel->IsNearestNeighborDataValid() ? LOCTEXT("Update", "Update") : LOCTEXT("Update *", "Update *");
 
@@ -110,6 +113,7 @@ namespace UE::NearestNeighborModel
 					if (NearestNeighborEditorModel != nullptr)
 					{
 						NearestNeighborEditorModel->UpdateNearestNeighborData();
+						NearestNeighborModel->InitPreviousWeights();
 						EditorModel->GetEditor()->GetModelDetailsView()->ForceRefresh();
 					}
 					return FReply::Handled();
@@ -131,6 +135,7 @@ namespace UE::NearestNeighborModel
 				})
 			];
 
+		MorphTargetCategoryBuilder->AddProperty(UNearestNeighborModel::GetMorphDataSizePropertyName());
 		MorphTargetCategoryBuilder->AddCustomRow(FText::FromString(""))
 			.WholeRowContent()
 			[
@@ -142,6 +147,9 @@ namespace UE::NearestNeighborModel
 					if (NearestNeighborEditorModel != nullptr)
 					{
 						NearestNeighborEditorModel->InitMorphTargets();
+						NearestNeighborEditorModel->RefreshMorphTargets();
+						NearestNeighborModel->UpdateNetworkSize();
+						NearestNeighborModel->UpdateMorphTargetSize();
 						EditorModel->GetEditor()->GetModelDetailsView()->ForceRefresh();
 					}
 					return FReply::Handled();
diff --git a/Engine/Plugins/Animation/MLDeformer/NearestNeighborModel/Source/NearestNeighborModelEditor/Public/NearestNeighborEditorModel.h b/Engine/Plugins/Animation/MLDeformer/NearestNeighborModel/Source/NearestNeighborModelEditor/Public/NearestNeighborEditorModel.h
--- a/Engine/Plugins/Animation/MLDeformer/NearestNeighborModel/Source/NearestNeighborModelEditor/Public/NearestNeighborEditorModel.h
+++ b/Engine/Plugins/Animation/MLDeformer/NearestNeighborModel/Source/NearestNeighborModelEditor/Public/NearestNeighborEditorModel.h
@@ -59,6 +59,7 @@ namespace UE::NearestNeighborModel
 		void KMeansClusterPoses();
 
 		void InitMorphTargets();
+		void RefreshMorphTargets();
 		void AddFloatArrayToDeltaArray(const TArray<float>& FloatArr, const TArray<uint32>& VertexMap, TArray<FVector3f>& DeltaArr, int32 DeltaArrayOffset = -1, float ScaleFactor = 1);
 
 		int32 GetNumParts();
diff --git a/Engine/Plugins/Animation/MLDeformer/NearestNeighborModel/Source/NearestNeighborModelEditor/Public/NearestNeighborEditorModelActor.h b/Engine/Plugins/Animation/MLDeformer/NearestNeighborModel/Source/NearestNeighborModelEditor/Public/NearestNeighborEditorModelActor.h
--- a/Engine/Plugins/Animation/MLDeformer/NearestNeighborModel/Source/NearestNeighborModelEditor/Public/NearestNeighborEditorModelActor.h
+++ b/Engine/Plugins/Animation/MLDeformer/NearestNeighborModel/Source/NearestNeighborModelEditor/Public/NearestNeighborEditorModelActor.h
@@ -3,7 +3,7 @@
 #pragma once
 
 #include "MLDeformerGeomCacheActor.h"
-#include "NearestNeighborModel.h"
+#include "NearestNeighborModelInstance.h"
 
 class UMLDeformerComponent;
 class UGeometryCacheComponent;
@@ -26,12 +26,12 @@ namespace UE::NearestNeighborModel
 		void SetGeometryCacheComponent(UGeometryCacheComponent* Component) { GeomCacheComponent = Component; }
 		UGeometryCacheComponent* GetGeometryCacheComponent() const { return GeomCacheComponent; }
 
-		void InitNearestNeighborActor(UNearestNeighborModel* InNearestNeighborModel, const int32 InPartId);
+		void InitNearestNeighborActor(UNearestNeighborModelInstance* InModelInstance, const int32 InPartId);
 		void TickNearestNeighborActor();
 
 	protected:
 		/** The geometry cache component (can be nullptr). */
-		TObjectPtr<UNearestNeighborModel> NearestNeighborModel = nullptr;
+		UNearestNeighborModelInstance* ModelInstance = nullptr;
 		int32 PartId = INDEX_NONE;
 	};
 }	// namespace UE::NearestNeighborModel
