Fix crash when NearestNeighborModelInstance::SetupNetwork returns false.

diff --git a/Engine/Plugins/Animation/MLDeformer/NearestNeighborModel/Source/NearestNeighborModel/Private/NearestNeighborModelInstance.cpp b/Engine/Plugins/Animation/MLDeformer/NearestNeighborModel/Source/NearestNeighborModel/Private/NearestNeighborModelInstance.cpp
--- a/Engine/Plugins/Animation/MLDeformer/NearestNeighborModel/Source/NearestNeighborModel/Private/NearestNeighborModelInstance.cpp
+++ b/Engine/Plugins/Animation/MLDeformer/NearestNeighborModel/Source/NearestNeighborModel/Private/NearestNeighborModelInstance.cpp
@@ -80,20 +80,21 @@ bool UNearestNeighborModelInstance::SetupInputs()
 {
     const bool bSuccess = Super::SetupInputs();
 
-    // Clip network inputs based on training set min and max
-    UNeuralNetwork* NeuralNetwork = Model->GetNeuralNetwork();
-    if (NeuralNetwork)
+    if (bSuccess)
     {
-		UNearestNeighborModel* NearestNeighborModel = static_cast<UNearestNeighborModel*>(Model.Get());
-    	float* InputDataPointer = static_cast<float*>(NeuralNetwork->GetInputDataPointerMutableForContext(NeuralNetworkInferenceHandle));
-    	const int64 NumNeuralNetInputs = NeuralNetwork->GetInputTensorForContext(NeuralNetworkInferenceHandle).Num();
-    	NearestNeighborModel->ClipInputs(InputDataPointer, NumNeuralNetInputs);
-    	return bSuccess;
-    }
-    else
-    {
-    	return false;
+	    UNeuralNetwork* NeuralNetwork = Model->GetNeuralNetwork();
+	    if (NeuralNetwork)
+	    {
+			UNearestNeighborModel* NearestNeighborModel = static_cast<UNearestNeighborModel*>(Model.Get());
+	    	float* InputDataPointer = static_cast<float*>(NeuralNetwork->GetInputDataPointerMutableForContext(NeuralNetworkInferenceHandle));
+	    	const int64 NumNeuralNetInputs = NeuralNetwork->GetInputTensorForContext(NeuralNetworkInferenceHandle).Num();
+		    // Clip network inputs based on training set min and max
+	    	NearestNeighborModel->ClipInputs(InputDataPointer, NumNeuralNetInputs);
+	    	return true;
+	    }
     }
+
+    return false;
 }
 
 void UNearestNeighborModelInstance::Execute(float ModelWeight)
