compile fixes for deprecated TObjectPtr.IsNull

diff --git a/Engine/Plugins/Animation/DeformerGraph/Source/OptimusSettings/Private/OptimusSettingsModule.cpp b/Engine/Plugins/Animation/DeformerGraph/Source/OptimusSettings/Private/OptimusSettingsModule.cpp
--- a/Engine/Plugins/Animation/DeformerGraph/Source/OptimusSettings/Private/OptimusSettingsModule.cpp
+++ b/Engine/Plugins/Animation/DeformerGraph/Source/OptimusSettings/Private/OptimusSettingsModule.cpp
@@ -66,11 +66,11 @@ TObjectPtr<UMeshDeformer> FOptimusSettingsModule::GetDefaultMeshDeformer(FDefaul
 		(Settings->DefaultMode == EOptimusDefaultDeformerMode::Always ||
 		(Settings->DefaultMode == EOptimusDefaultDeformerMode::SkinCacheOnly && Setup.bIsUsingSkinCache)))
 	{
-		if (Setup.bIsRequestingRecomputeTangent && !DefaultRecomputeTangentDeformer.IsNull())
+		if (Setup.bIsRequestingRecomputeTangent && DefaultRecomputeTangentDeformer)
 		{
 			return DefaultRecomputeTangentDeformer;
 		}
-		if (!DefaultDeformer.IsNull())
+		if (DefaultDeformer)
 		{ 
 			return DefaultDeformer;
 		}
diff --git a/Engine/Plugins/Animation/MLDeformer/MLDeformerFramework/Source/MLDeformerFramework/Private/MLDeformerMorphModelInstance.cpp b/Engine/Plugins/Animation/MLDeformer/MLDeformerFramework/Source/MLDeformerFramework/Private/MLDeformerMorphModelInstance.cpp
--- a/Engine/Plugins/Animation/MLDeformer/MLDeformerFramework/Source/MLDeformerFramework/Private/MLDeformerMorphModelInstance.cpp
+++ b/Engine/Plugins/Animation/MLDeformer/MLDeformerFramework/Source/MLDeformerFramework/Private/MLDeformerMorphModelInstance.cpp
@@ -18,7 +18,7 @@ void UMLDeformerMorphModelInstance::Init(USkeletalMeshComponent* SkelMeshCompone
 void UMLDeformerMorphModelInstance::Release()
 {
 	// Try to unregister the morph target and morph target set.
-	if (!SkeletalMeshComponent.IsNull())
+	if (SkeletalMeshComponent)
 	{
 		const UMLDeformerMorphModel* MorphModel = Cast<UMLDeformerMorphModel>(Model);
 		if (MorphModel)
@@ -92,7 +92,7 @@ void UMLDeformerMorphModelInstance::PostMLDeformerComponentInit()
 FExternalMorphSetWeights* UMLDeformerMorphModelInstance::FindWeightData(int32 LOD) const
 {
 	const UMLDeformerMorphModel* MorphModel = Cast<UMLDeformerMorphModel>(Model);
-	if (MorphModel == nullptr || SkeletalMeshComponent.IsNull())
+	if (MorphModel == nullptr || !SkeletalMeshComponent)
 	{
 		return nullptr;
 	}
diff --git a/Engine/Plugins/Animation/MLDeformer/MLDeformerFramework/Source/MLDeformerFrameworkEditor/Public/MLDeformerEditorModel.h b/Engine/Plugins/Animation/MLDeformer/MLDeformerFramework/Source/MLDeformerFrameworkEditor/Public/MLDeformerEditorModel.h
--- a/Engine/Plugins/Animation/MLDeformer/MLDeformerFramework/Source/MLDeformerFrameworkEditor/Public/MLDeformerEditorModel.h
+++ b/Engine/Plugins/Animation/MLDeformer/MLDeformerFramework/Source/MLDeformerFrameworkEditor/Public/MLDeformerEditorModel.h
@@ -86,7 +86,7 @@ namespace UE::MLDeformer
 
 			// Update to the new asset, and mark it as editor only if it isn't yet.
 			Asset = InAsset;
-			if (!Asset.IsNull())
+			if (Asset)
 			{
 				UPackage* Package = Asset->GetPackage();
 				bEditorOnly = static_cast<bool>(Package->GetPackageFlags() & PKG_EditorOnly);
diff --git a/Engine/Plugins/Animation/MLDeformer/NearestNeighborModel/Source/NearestNeighborModel/Private/NearestNeighborModel.cpp b/Engine/Plugins/Animation/MLDeformer/NearestNeighborModel/Source/NearestNeighborModel/Private/NearestNeighborModel.cpp
--- a/Engine/Plugins/Animation/MLDeformer/NearestNeighborModel/Source/NearestNeighborModel/Private/NearestNeighborModel.cpp
+++ b/Engine/Plugins/Animation/MLDeformer/NearestNeighborModel/Source/NearestNeighborModel/Private/NearestNeighborModel.cpp
@@ -202,7 +202,7 @@ const TObjectPtr<UGeometryCache> UNearestNeighborModel::GetNearestNeighborCache(
 int32 UNearestNeighborModel::GetNumNeighborsFromGeometryCache(int32 PartId) const
 {
 	const TObjectPtr<UGeometryCache> Cache = GetNearestNeighborCache(PartId);
-	if (!Cache.IsNull())
+	if (Cache)
 	{
 		const int32 StartFrame = Cache->GetStartFrame();
 		const int32 EndFrame = Cache->GetEndFrame();
diff --git a/Engine/Plugins/Animation/MLDeformer/NearestNeighborModel/Source/NearestNeighborModelEditor/Private/NearestNeighborEditorModel.cpp b/Engine/Plugins/Animation/MLDeformer/NearestNeighborModel/Source/NearestNeighborModelEditor/Private/NearestNeighborEditorModel.cpp
--- a/Engine/Plugins/Animation/MLDeformer/NearestNeighborModel/Source/NearestNeighborModelEditor/Private/NearestNeighborEditorModel.cpp
+++ b/Engine/Plugins/Animation/MLDeformer/NearestNeighborModel/Source/NearestNeighborModelEditor/Private/NearestNeighborEditorModel.cpp
@@ -172,7 +172,7 @@ namespace UE::NearestNeighborModel
 			const TObjectPtr<UAnimSequence> AnimSequence = NearestNeighborModel->GetNearestNeighborSkeletons(PartId);
 			const TObjectPtr<UGeometryCache> GeometryCache = NearestNeighborModel->GetNearestNeighborCache(PartId);
 
-			if (!SkeletalMeshComponent.IsNull() && !GeometryCacheComponent.IsNull() && !AnimSequence.IsNull() && !GeometryCache.IsNull())
+			if (SkeletalMeshComponent && GeometryCacheComponent && AnimSequence && GeometryCache)
 			{
 				const int32 NumFrames = GeometryCache->GetEndFrame() - GeometryCache->GetStartFrame() + 1;
 				if (AnimSequence->GetDataModel()->GetNumberOfKeys() == NumFrames)
@@ -233,7 +233,7 @@ namespace UE::NearestNeighborModel
 			TObjectPtr<USkeletalMeshComponent> SkeletalMeshComponent = GeomCacheSampler->GetSkeletalMeshComponent();
 			const TObjectPtr<UAnimSequence> AnimSequence = NearestNeighborModel->GetAnimSequence();
 
-			if (!SkeletalMeshComponent.IsNull() && !AnimSequence.IsNull())
+			if (SkeletalMeshComponent && AnimSequence)
 			{
 				SkeletalMeshComponent->SetAnimationMode(EAnimationMode::AnimationSingleNode);
 				SkeletalMeshComponent->SetAnimation(AnimSequence);
@@ -304,7 +304,7 @@ namespace UE::NearestNeighborModel
 			TObjectPtr<USkeletalMeshComponent> SkeletalMeshComponent = GeomCacheSampler->GetSkeletalMeshComponent();
 			TObjectPtr<UGeometryCacheComponent> GeometryCacheComponent = GeomCacheSampler->GetGeometryCacheComponent();
 
-			if(!SkeletalMeshComponent.IsNull() && !GeometryCacheComponent.IsNull())
+			if(SkeletalMeshComponent && GeometryCacheComponent)
 			{
 				UNearestNeighborTrainingModel *TrainingModel = InitTrainingModel<UNearestNeighborTrainingModel>(this);
 				check(TrainingModel != nullptr);
diff --git a/Engine/Plugins/Animation/MLDeformer/NearestNeighborModel/Source/NearestNeighborModelEditor/Private/NearestNeighborGeomCacheSampler.cpp b/Engine/Plugins/Animation/MLDeformer/NearestNeighborModel/Source/NearestNeighborModelEditor/Private/NearestNeighborGeomCacheSampler.cpp
--- a/Engine/Plugins/Animation/MLDeformer/NearestNeighborModel/Source/NearestNeighborModelEditor/Private/NearestNeighborGeomCacheSampler.cpp
+++ b/Engine/Plugins/Animation/MLDeformer/NearestNeighborModel/Source/NearestNeighborModelEditor/Private/NearestNeighborGeomCacheSampler.cpp
@@ -140,7 +140,7 @@ namespace UE::NearestNeighborModel
 	void FNearestNeighborGeomCacheSampler::SampleKMeansAnim(const int32 SkeletonId)
 	{
 		UNearestNeighborModel* NearestNeighborModel = static_cast<UNearestNeighborModel*>(Model);
-		if (NearestNeighborModel && SkeletonId < NearestNeighborModel->SourceSkeletons.Num() && !SkeletalMeshComponent.IsNull())
+		if (NearestNeighborModel && SkeletonId < NearestNeighborModel->SourceSkeletons.Num() && SkeletalMeshComponent)
 		{
 			const TObjectPtr<UAnimSequence> AnimSequence = NearestNeighborModel->SourceSkeletons[SkeletonId];
 			SkeletalMeshComponent->SetAnimationMode(EAnimationMode::AnimationSingleNode);
diff --git a/Engine/Source/Editor/SequenceRecorder/Private/AnimationRecorder.cpp b/Engine/Source/Editor/SequenceRecorder/Private/AnimationRecorder.cpp
--- a/Engine/Source/Editor/SequenceRecorder/Private/AnimationRecorder.cpp
+++ b/Engine/Source/Editor/SequenceRecorder/Private/AnimationRecorder.cpp
@@ -235,7 +235,7 @@ void FAnimationRecorder::StartRecord(USkeletalMeshComponent* Component, UAnimSeq
 	TimePassed = 0.0;
 	AnimationObject = InAnimationObject;
 
-	if (AnimationObject->BoneCompressionSettings.IsNull())
+	if (!AnimationObject->BoneCompressionSettings)
 	{
 		AnimationObject->BoneCompressionSettings = FAnimationUtils::GetDefaultAnimationRecorderBoneCompressionSettings();
 	}
diff --git a/Engine/Source/Runtime/AudioMixer/Private/AudioMixerSubmix.cpp b/Engine/Source/Runtime/AudioMixer/Private/AudioMixerSubmix.cpp
--- a/Engine/Source/Runtime/AudioMixer/Private/AudioMixerSubmix.cpp
+++ b/Engine/Source/Runtime/AudioMixer/Private/AudioMixerSubmix.cpp
@@ -294,7 +294,7 @@ namespace Audio
 					check(LinkFactory->GetSettingsClass());
 					check(LinkFactory->GetSettingsClass()->GetDefaultObject());
 
-					const UAudioLinkSettingsAbstract* Settings = SoundSubmix->AudioLinkSettings.IsNull() ?
+					const UAudioLinkSettingsAbstract* Settings = !SoundSubmix->AudioLinkSettings ?
 						GetDefault<UAudioLinkSettingsAbstract>(LinkFactory->GetSettingsClass()) : SoundSubmix->AudioLinkSettings.Get();
 
 					AudioLinkInstance = LinkFactory->CreateSubmixAudioLink({ SoundSubmix, MixerDevice, Settings });
