Fix crash when updating nearest neighbor dataset.
Move GetInstance to NearestNeighborModelActor

diff --git a/Engine/Plugins/Animation/MLDeformer/NearestNeighborModel/Source/NearestNeighborModel/Private/NearestNeighborModelInstance.cpp b/Engine/Plugins/Animation/MLDeformer/NearestNeighborModel/Source/NearestNeighborModel/Private/NearestNeighborModelInstance.cpp
--- a/Engine/Plugins/Animation/MLDeformer/NearestNeighborModel/Source/NearestNeighborModel/Private/NearestNeighborModelInstance.cpp
+++ b/Engine/Plugins/Animation/MLDeformer/NearestNeighborModel/Source/NearestNeighborModel/Private/NearestNeighborModelInstance.cpp
@@ -153,7 +153,11 @@ void UNearestNeighborModelInstance::RunNearestNeighborModel(float ModelWeight)
 				for (int32 NeighborId = 0; NeighborId < NumNeighbors; NeighborId++)
 				{
 					const float W = NeighborId == NearestNeighborId ? ModelWeight * NearestNeighborModel->GetNearestNeighborOffsetWeight() : 0;
-					UpdateWeight(WeightData->Weights, NeighborOffset + NeighborId, W);
+					const int32 Index = NeighborOffset + NeighborId;
+					if (Index < NumMorphTargets)
+					{
+						UpdateWeight(WeightData->Weights, Index, W);
+					}
 				}
 				NeighborOffset += NumNeighbors;
 			}
diff --git a/Engine/Plugins/Animation/MLDeformer/NearestNeighborModel/Source/NearestNeighborModelEditor/Private/NearestNeighborEditorModel.cpp b/Engine/Plugins/Animation/MLDeformer/NearestNeighborModel/Source/NearestNeighborModelEditor/Private/NearestNeighborEditorModel.cpp
--- a/Engine/Plugins/Animation/MLDeformer/NearestNeighborModel/Source/NearestNeighborModelEditor/Private/NearestNeighborEditorModel.cpp
+++ b/Engine/Plugins/Animation/MLDeformer/NearestNeighborModel/Source/NearestNeighborModelEditor/Private/NearestNeighborEditorModel.cpp
@@ -122,22 +122,7 @@ namespace UE::NearestNeighborModel
 				false);
 			FNearestNeighborEditorModelActor* NearestNeighborActor = static_cast<FNearestNeighborEditorModelActor*>(EditorActors.Last());
 
-			UNearestNeighborModelInstance* ModelInstance = nullptr;
-			const FMLDeformerEditorActor* EditorActor = FindEditorActor(ActorID_Test_MLDeformed);
-			if (EditorActor)
-			{
-				const AActor* Actor = EditorActor->GetActor();
-				if (Actor)
-				{
-					UMLDeformerComponent* Component = Actor->FindComponentByClass<UMLDeformerComponent>();
-					if (Component)
-					{
-						ModelInstance = static_cast<UNearestNeighborModelInstance*>(Component->GetModelInstance());
-					}
-				}
-			}
-
-			NearestNeighborActor->InitNearestNeighborActor(ModelInstance, PartId);
+			NearestNeighborActor->InitNearestNeighborActor(PartId);
 			NearestNeighborActor->SetMeshOffsetFactor(2.0f);
 			NearestNeighborActors[PartId] = NearestNeighborActor;
 		}
diff --git a/Engine/Plugins/Animation/MLDeformer/NearestNeighborModel/Source/NearestNeighborModelEditor/Private/NearestNeighborEditorModelActor.cpp b/Engine/Plugins/Animation/MLDeformer/NearestNeighborModel/Source/NearestNeighborModelEditor/Private/NearestNeighborEditorModelActor.cpp
--- a/Engine/Plugins/Animation/MLDeformer/NearestNeighborModel/Source/NearestNeighborModelEditor/Private/NearestNeighborEditorModelActor.cpp
+++ b/Engine/Plugins/Animation/MLDeformer/NearestNeighborModel/Source/NearestNeighborModelEditor/Private/NearestNeighborEditorModelActor.cpp
@@ -1,6 +1,7 @@
 // Copyright Epic Games, Inc. All Rights Reserved.
 
 #include "NearestNeighborEditorModelActor.h"
+#include "MLDeformerComponent.h"
 #include "GeometryCacheComponent.h"
 
 namespace UE::NearestNeighborModel
@@ -12,18 +13,33 @@ namespace UE::NearestNeighborModel
 	{
 	}
 
-	void FNearestNeighborEditorModelActor::InitNearestNeighborActor(UNearestNeighborModelInstance* InModelInstance, const int32 InPartId)
+	void FNearestNeighborEditorModelActor::InitNearestNeighborActor(const int32 InPartId)
 	{
-		ModelInstance = InModelInstance;
 		PartId = InPartId;
 	}
 
 	void FNearestNeighborEditorModelActor::TickNearestNeighborActor()
 	{
+		const UNearestNeighborModelInstance* ModelInstance = GetModelInstance();
 		if (GeomCacheComponent && GeomCacheComponent->GetGeometryCache() && ModelInstance && PartId < ModelInstance->NeighborIdNum())
 		{
 			GeomCacheComponent->SetManualTick(true);
 			GeomCacheComponent->TickAtThisTime(GeomCacheComponent->GetTimeAtFrame(ModelInstance->NearestNeighborId(PartId)), false, false, false);
 		}
 	}
+
+	UNearestNeighborModelInstance* FNearestNeighborEditorModelActor::GetModelInstance() const
+	{
+		UNearestNeighborModelInstance* ModelInstance = nullptr;
+		const AActor* MyActor = GetActor();
+		if (MyActor)
+		{
+			UMLDeformerComponent* Component = MyActor->FindComponentByClass<UMLDeformerComponent>();
+			if (Component)
+			{
+				ModelInstance = static_cast<UNearestNeighborModelInstance*>(Component->GetModelInstance());
+			}
+		}
+		return ModelInstance;
+	}
 }	// namespace UE::NearestNeighborModel
diff --git a/Engine/Plugins/Animation/MLDeformer/NearestNeighborModel/Source/NearestNeighborModelEditor/Public/NearestNeighborEditorModelActor.h b/Engine/Plugins/Animation/MLDeformer/NearestNeighborModel/Source/NearestNeighborModelEditor/Public/NearestNeighborEditorModelActor.h
--- a/Engine/Plugins/Animation/MLDeformer/NearestNeighborModel/Source/NearestNeighborModelEditor/Public/NearestNeighborEditorModelActor.h
+++ b/Engine/Plugins/Animation/MLDeformer/NearestNeighborModel/Source/NearestNeighborModelEditor/Public/NearestNeighborEditorModelActor.h
@@ -26,12 +26,12 @@ namespace UE::NearestNeighborModel
 		void SetGeometryCacheComponent(UGeometryCacheComponent* Component) { GeomCacheComponent = Component; }
 		UGeometryCacheComponent* GetGeometryCacheComponent() const { return GeomCacheComponent; }
 
-		void InitNearestNeighborActor(UNearestNeighborModelInstance* InModelInstance, const int32 InPartId);
+		void InitNearestNeighborActor(const int32 InPartId);
 		void TickNearestNeighborActor();
 
+		UNearestNeighborModelInstance* GetModelInstance() const;
+
 	protected:
-		/** The geometry cache component (can be nullptr). */
-		UNearestNeighborModelInstance* ModelInstance = nullptr;
 		int32 PartId = INDEX_NONE;
 	};
 }	// namespace UE::NearestNeighborModel
