Fix ticking bug nearest neighbor actors

diff --git a/Engine/Plugins/Animation/MLDeformer/NearestNeighborModel/Source/NearestNeighborModelEditor/Private/NearestNeighborEditorModel.cpp b/Engine/Plugins/Animation/MLDeformer/NearestNeighborModel/Source/NearestNeighborModelEditor/Private/NearestNeighborEditorModel.cpp
--- a/Engine/Plugins/Animation/MLDeformer/NearestNeighborModel/Source/NearestNeighborModelEditor/Private/NearestNeighborEditorModel.cpp
+++ b/Engine/Plugins/Animation/MLDeformer/NearestNeighborModel/Source/NearestNeighborModelEditor/Private/NearestNeighborEditorModel.cpp
@@ -106,6 +106,18 @@ namespace UE::NearestNeighborModel
 			NearestNeighborActors.Reset();
 		}
 		NearestNeighborActors.SetNumZeroed(NumParts);
+
+		UMLDeformerComponent* MLDeformerComponent = nullptr;
+		FMLDeformerEditorActor* TestActor = FindEditorActor(ActorID_Test_MLDeformed);
+		if (TestActor)
+		{
+			AActor* Actor = TestActor->GetActor();
+			if (Actor)
+			{
+				MLDeformerComponent = Actor->FindComponentByClass<UMLDeformerComponent>();
+			}
+		}
+		
 		for(int32 PartId = StartIndex; PartId < NumParts; PartId++)
 		{
 			UGeometryCache* GeomCache = NearestNeighborModel->GetNearestNeighborCache(PartId);
@@ -122,7 +134,7 @@ namespace UE::NearestNeighborModel
 				false);
 			FNearestNeighborEditorModelActor* NearestNeighborActor = static_cast<FNearestNeighborEditorModelActor*>(EditorActors.Last());
 
-			NearestNeighborActor->InitNearestNeighborActor(PartId);
+			NearestNeighborActor->InitNearestNeighborActor(PartId, MLDeformerComponent);
 			NearestNeighborActor->SetMeshOffsetFactor(2.0f);
 			NearestNeighborActors[PartId] = NearestNeighborActor;
 		}
diff --git a/Engine/Plugins/Animation/MLDeformer/NearestNeighborModel/Source/NearestNeighborModelEditor/Private/NearestNeighborEditorModelActor.cpp b/Engine/Plugins/Animation/MLDeformer/NearestNeighborModel/Source/NearestNeighborModelEditor/Private/NearestNeighborEditorModelActor.cpp
--- a/Engine/Plugins/Animation/MLDeformer/NearestNeighborModel/Source/NearestNeighborModelEditor/Private/NearestNeighborEditorModelActor.cpp
+++ b/Engine/Plugins/Animation/MLDeformer/NearestNeighborModel/Source/NearestNeighborModelEditor/Private/NearestNeighborEditorModelActor.cpp
@@ -1,7 +1,7 @@
 // Copyright Epic Games, Inc. All Rights Reserved.
 
 #include "NearestNeighborEditorModelActor.h"
-#include "MLDeformerComponent.h"
+#include "NearestNeighborModelInstance.h"
 #include "GeometryCacheComponent.h"
 
 namespace UE::NearestNeighborModel
@@ -13,33 +13,22 @@ namespace UE::NearestNeighborModel
 	{
 	}
 
-	void FNearestNeighborEditorModelActor::InitNearestNeighborActor(const int32 InPartId)
+	void FNearestNeighborEditorModelActor::InitNearestNeighborActor(const int32 InPartId, const UMLDeformerComponent* InComponent)
 	{
 		PartId = InPartId;
+		MLDeformerComponent = InComponent;
 	}
 
 	void FNearestNeighborEditorModelActor::TickNearestNeighborActor()
 	{
-		const UNearestNeighborModelInstance* ModelInstance = GetModelInstance();
-		if (GeomCacheComponent && GeomCacheComponent->GetGeometryCache() && ModelInstance && PartId < ModelInstance->NeighborIdNum())
+		if (GeomCacheComponent && GeomCacheComponent->GetGeometryCache() && MLDeformerComponent)
 		{
-			GeomCacheComponent->SetManualTick(true);
-			GeomCacheComponent->TickAtThisTime(GeomCacheComponent->GetTimeAtFrame(ModelInstance->NearestNeighborId(PartId)), false, false, false);
-		}
-	}
-
-	UNearestNeighborModelInstance* FNearestNeighborEditorModelActor::GetModelInstance() const
-	{
-		UNearestNeighborModelInstance* ModelInstance = nullptr;
-		const AActor* MyActor = GetActor();
-		if (MyActor)
-		{
-			UMLDeformerComponent* Component = MyActor->FindComponentByClass<UMLDeformerComponent>();
-			if (Component)
+			const UNearestNeighborModelInstance* ModelInstance = static_cast<UNearestNeighborModelInstance*>(MLDeformerComponent->GetModelInstance());
+			if (ModelInstance && PartId < ModelInstance->NeighborIdNum())
 			{
-				ModelInstance = static_cast<UNearestNeighborModelInstance*>(Component->GetModelInstance());
+				GeomCacheComponent->SetManualTick(true);
+				GeomCacheComponent->TickAtThisTime(GeomCacheComponent->GetTimeAtFrame(ModelInstance->NearestNeighborId(PartId)), false, false, false);
 			}
 		}
-		return ModelInstance;
 	}
 }	// namespace UE::NearestNeighborModel
diff --git a/Engine/Plugins/Animation/MLDeformer/NearestNeighborModel/Source/NearestNeighborModelEditor/Public/NearestNeighborEditorModelActor.h b/Engine/Plugins/Animation/MLDeformer/NearestNeighborModel/Source/NearestNeighborModelEditor/Public/NearestNeighborEditorModelActor.h
--- a/Engine/Plugins/Animation/MLDeformer/NearestNeighborModel/Source/NearestNeighborModelEditor/Public/NearestNeighborEditorModelActor.h
+++ b/Engine/Plugins/Animation/MLDeformer/NearestNeighborModel/Source/NearestNeighborModelEditor/Public/NearestNeighborEditorModelActor.h
@@ -3,7 +3,7 @@
 #pragma once
 
 #include "MLDeformerGeomCacheActor.h"
-#include "NearestNeighborModelInstance.h"
+#include "MLDeformerComponent.h"
 
 class UMLDeformerComponent;
 class UGeometryCacheComponent;
@@ -26,12 +26,11 @@ namespace UE::NearestNeighborModel
 		void SetGeometryCacheComponent(UGeometryCacheComponent* Component) { GeomCacheComponent = Component; }
 		UGeometryCacheComponent* GetGeometryCacheComponent() const { return GeomCacheComponent; }
 
-		void InitNearestNeighborActor(const int32 InPartId);
+		void InitNearestNeighborActor(const int32 InPartId, const UMLDeformerComponent* InComponent);
 		void TickNearestNeighborActor();
 
-		UNearestNeighborModelInstance* GetModelInstance() const;
-
 	protected:
 		int32 PartId = INDEX_NONE;
+		const UMLDeformerComponent* MLDeformerComponent = nullptr;
 	};
 }	// namespace UE::NearestNeighborModel
