Add file cache in NearestNeighborModel.
This allows users to change bone list, num of PCA basis and other training settings without re-generate deltas or PCA components.

diff --git a/Engine/Plugins/Animation/MLDeformer/NearestNeighborModel/Source/NearestNeighborModel/Private/NearestNeighborModel.cpp b/Engine/Plugins/Animation/MLDeformer/NearestNeighborModel/Source/NearestNeighborModel/Private/NearestNeighborModel.cpp
--- a/Engine/Plugins/Animation/MLDeformer/NearestNeighborModel/Source/NearestNeighborModel/Private/NearestNeighborModel.cpp
+++ b/Engine/Plugins/Animation/MLDeformer/NearestNeighborModel/Source/NearestNeighborModel/Private/NearestNeighborModel.cpp
@@ -247,6 +247,18 @@ void UNearestNeighborModel::UpdateMorphTargetSize()
 		MorphDataSize = 0.0f;
 	}
 }
+
+FString UNearestNeighborModel::GetModelDir() const
+{
+	if (bUseFileCache)
+	{
+		return FileCacheDirectory;
+	}
+	else
+	{
+		return FPaths::ProjectIntermediateDir() + TEXT("NearestNeighborModel/");
+	}
+}
 #endif
 
 void UNearestNeighborModel::InitInputInfo()
diff --git a/Engine/Plugins/Animation/MLDeformer/NearestNeighborModel/Source/NearestNeighborModel/Public/NearestNeighborModel.h b/Engine/Plugins/Animation/MLDeformer/NearestNeighborModel/Source/NearestNeighborModel/Public/NearestNeighborModel.h
--- a/Engine/Plugins/Animation/MLDeformer/NearestNeighborModel/Source/NearestNeighborModel/Public/NearestNeighborModel.h
+++ b/Engine/Plugins/Animation/MLDeformer/NearestNeighborModel/Source/NearestNeighborModel/Public/NearestNeighborModel.h
@@ -265,6 +265,15 @@ public:
 
 	static FName GetSavedNetworkSizePropertyName() { return GET_MEMBER_NAME_CHECKED(UNearestNeighborModel, SavedNetworkSize); }
 	static FName GetMorphDataSizePropertyName() { return GET_MEMBER_NAME_CHECKED(UNearestNeighborModel, MorphDataSize); }
+
+	bool GetUseFileCache() const { return bUseFileCache; }
+	static FName GetUseFileCachePropertyName() { return GET_MEMBER_NAME_CHECKED(UNearestNeighborModel, bUseFileCache); }
+	static FName GetFileCacheDirectoryPropertyName() { return GET_MEMBER_NAME_CHECKED(UNearestNeighborModel, FileCacheDirectory); }
+	static FName GetRecomputeDeltasPropertyName() { return GET_MEMBER_NAME_CHECKED(UNearestNeighborModel, bRecomputeDeltas); }
+	static FName GetRecomputePCAPropertyName() { return GET_MEMBER_NAME_CHECKED(UNearestNeighborModel, bRecomputePCA); }
+
+	UFUNCTION(BlueprintPure, Category = "Nearest Neighbor Model")
+	FString GetModelDir() const;
 #endif
 
 
@@ -305,6 +314,18 @@ protected:
 
 	UPROPERTY(VisibleAnywhere, BlueprintReadOnly, Category = "Morph Targets", meta = (ClampMin = "0.0"))
 	float MorphDataSize = 0.0f;
+
+	UPROPERTY(EditAnywhere, BlueprintReadWrite, Category = "File Cache")
+	bool bUseFileCache = false;
+
+	UPROPERTY(EditAnywhere, BlueprintReadWrite, Category = "File Cache")
+	FString FileCacheDirectory = FPaths::ProjectIntermediateDir() + "NearestNeighborModel";
+
+	UPROPERTY(EditAnywhere, BlueprintReadWrite, Category = "File Cache")
+	bool bRecomputeDeltas = true;
+
+	UPROPERTY(EditAnywhere, BlueprintReadWrite, Category = "File Cache")
+	bool bRecomputePCA = true;
 #endif
 
 public:
diff --git a/Engine/Plugins/Animation/MLDeformer/NearestNeighborModel/Source/NearestNeighborModelEditor/Private/NearestNeighborEditorModel.cpp b/Engine/Plugins/Animation/MLDeformer/NearestNeighborModel/Source/NearestNeighborModelEditor/Private/NearestNeighborEditorModel.cpp
--- a/Engine/Plugins/Animation/MLDeformer/NearestNeighborModel/Source/NearestNeighborModelEditor/Private/NearestNeighborEditorModel.cpp
+++ b/Engine/Plugins/Animation/MLDeformer/NearestNeighborModel/Source/NearestNeighborModelEditor/Private/NearestNeighborEditorModel.cpp
@@ -75,6 +75,19 @@ namespace UE::NearestNeighborModel
 		GetNearestNeighborModel()->InitPreviousWeights();
 	}
 
+	FString FNearestNeighborEditorModel::GetTrainedNetworkOnnxFile() const
+	{
+		const UNearestNeighborModel* NearestNeighborModel = GetNearestNeighborModel();
+		if (NearestNeighborModel)
+		{
+			return NearestNeighborModel->GetModelDir() + TEXT("/NearestNeighborModel.onnx");
+		}
+		else
+		{
+			return FMLDeformerEditorModel::GetTrainedNetworkOnnxFile();
+		}
+	}
+
 	void FNearestNeighborEditorModel::CreateActors(const TSharedRef<IPersonaPreviewScene>& InPersonaPreviewScene)
 	{
 		FMLDeformerMorphModelEditorModel::CreateActors(InPersonaPreviewScene);
@@ -282,8 +295,8 @@ namespace UE::NearestNeighborModel
 		UNeuralNetwork *NeuralNetwork = NearestNeighborModel->GetNeuralNetwork();
 		if(NeuralNetwork && NeuralNetwork->IsLoaded())
 		{
-			const FString SavePath = FPaths::ProjectIntermediateDir() + "NearestNeighborModel/NearestNeighborModel.onnx";
-			if (!FPaths::FileExists(SavePath))
+			const FString SavePath = GetTrainedNetworkOnnxFile();
+			if (!FPaths::FileExists(SavePath) || !NearestNeighborModel->GetUseFileCache())
 			{
 				UE_LOG(LogNearestNeighborModel, Display, TEXT("Saving to %s"), *SavePath);
 				NeuralNetwork->Save(SavePath);
@@ -321,8 +334,8 @@ namespace UE::NearestNeighborModel
 		UNeuralNetwork *NeuralNetwork = NearestNeighborModel->GetNeuralNetwork();
 		if(NeuralNetwork && NeuralNetwork->IsLoaded())
 		{
-			const FString SavePath = FPaths::ProjectIntermediateDir() + "NearestNeighborModel/NearestNeighborModel.onnx";
-			if (!FPaths::FileExists(SavePath))
+			const FString SavePath = GetTrainedNetworkOnnxFile();
+			if (!FPaths::FileExists(SavePath) || !NearestNeighborModel->GetUseFileCache())
 			{
 				UE_LOG(LogNearestNeighborModel, Display, TEXT("Saving to %s"), *SavePath);
 				NeuralNetwork->Save(SavePath);
diff --git a/Engine/Plugins/Animation/MLDeformer/NearestNeighborModel/Source/NearestNeighborModelEditor/Private/NearestNeighborModelDetails.cpp b/Engine/Plugins/Animation/MLDeformer/NearestNeighborModel/Source/NearestNeighborModelEditor/Private/NearestNeighborModelDetails.cpp
--- a/Engine/Plugins/Animation/MLDeformer/NearestNeighborModel/Source/NearestNeighborModelEditor/Private/NearestNeighborModelDetails.cpp
+++ b/Engine/Plugins/Animation/MLDeformer/NearestNeighborModel/Source/NearestNeighborModelEditor/Private/NearestNeighborModelDetails.cpp
@@ -8,6 +8,7 @@
 #include "DetailCategoryBuilder.h"
 #include "DetailWidgetRow.h"
 #include "SWarningOrErrorBox.h"
+#include "IDetailGroup.h"
 #include "IDetailsView.h"
 #include "Widgets/Input/SButton.h"
 
@@ -39,6 +40,7 @@ namespace UE::NearestNeighborModel
 	void FNearestNeighborModelDetails::CreateCategories()
 	{
 		FMLDeformerMorphModelDetails::CreateCategories();
+		FileCacheCategoryBuilder = &DetailLayoutBuilder->EditCategory("File Cache", FText::GetEmpty(), ECategoryPriority::Important);
 		ClothPartCategoryBuilder = &DetailLayoutBuilder->EditCategory("Cloth Parts", FText::GetEmpty(), ECategoryPriority::Important);
 		NearestNeighborCategoryBuilder = &DetailLayoutBuilder->EditCategory("Nearest Neighbors", FText::GetEmpty(), ECategoryPriority::Important);
 		KMeansCategoryBuilder = &DetailLayoutBuilder->EditCategory("KMeans Pose Generator", FText::GetEmpty(), ECategoryPriority::Important);
@@ -71,6 +73,13 @@ namespace UE::NearestNeighborModel
 		TrainingSettingsCategoryBuilder->AddProperty(UNearestNeighborModel::GetLearningRatePropertyName());
 		TrainingSettingsCategoryBuilder->AddProperty(UNearestNeighborModel::GetSavedNetworkSizePropertyName());
 
+		// File cache settings
+		IDetailGroup *Group = &FileCacheCategoryBuilder->AddGroup(TEXT("File Cache"), LOCTEXT("File Cache", "File Cache"), false, true);
+		Group->HeaderProperty(DetailBuilder.GetProperty(UNearestNeighborModel::GetUseFileCachePropertyName()));
+		Group->AddPropertyRow(DetailBuilder.GetProperty(UNearestNeighborModel::GetFileCacheDirectoryPropertyName()));
+		Group->AddPropertyRow(DetailBuilder.GetProperty(UNearestNeighborModel::GetRecomputeDeltasPropertyName()));
+		Group->AddPropertyRow(DetailBuilder.GetProperty(UNearestNeighborModel::GetRecomputePCAPropertyName()));
+
 		// Cloth part settings
 		ClothPartCategoryBuilder->AddProperty(UNearestNeighborModel::GetClothPartEditorDataPropertyName());
 		FText ButtonText = (NearestNeighborModel && NearestNeighborModel->IsClothPartDataValid()) ? LOCTEXT("Update", "Update") : LOCTEXT("Update *", "Update *");
diff --git a/Engine/Plugins/Animation/MLDeformer/NearestNeighborModel/Source/NearestNeighborModelEditor/Public/NearestNeighborEditorModel.h b/Engine/Plugins/Animation/MLDeformer/NearestNeighborModel/Source/NearestNeighborModelEditor/Public/NearestNeighborEditorModel.h
--- a/Engine/Plugins/Animation/MLDeformer/NearestNeighborModel/Source/NearestNeighborModelEditor/Public/NearestNeighborEditorModel.h
+++ b/Engine/Plugins/Animation/MLDeformer/NearestNeighborModel/Source/NearestNeighborModelEditor/Public/NearestNeighborEditorModel.h
@@ -41,6 +41,7 @@ namespace UE::NearestNeighborModel
 		virtual ETrainingResult Train() override;
 		virtual void OnPropertyChanged(FPropertyChangedEvent& PropertyChangedEvent) override;
 		virtual void OnPostTraining(ETrainingResult TrainingResult, bool bUsePartiallyTrainedWhenAborted) override;
+		virtual FString GetTrainedNetworkOnnxFile() const override;
 		// ~END FMLDeformerGeomCacheEditorModel overrides.
 
 		// Some helpers that cast to this model's variants of some classes.
diff --git a/Engine/Plugins/Animation/MLDeformer/NearestNeighborModel/Source/NearestNeighborModelEditor/Public/NearestNeighborModelDetails.h b/Engine/Plugins/Animation/MLDeformer/NearestNeighborModel/Source/NearestNeighborModelEditor/Public/NearestNeighborModelDetails.h
--- a/Engine/Plugins/Animation/MLDeformer/NearestNeighborModel/Source/NearestNeighborModelEditor/Public/NearestNeighborModelDetails.h
+++ b/Engine/Plugins/Animation/MLDeformer/NearestNeighborModel/Source/NearestNeighborModelEditor/Public/NearestNeighborModelDetails.h
@@ -30,6 +30,7 @@ namespace UE::NearestNeighborModel
 		UNearestNeighborModel* NearestNeighborModel = nullptr;
 		FNearestNeighborEditorModel* NearestNeighborEditorModel = nullptr;
 
+		IDetailCategoryBuilder* FileCacheCategoryBuilder = nullptr;
 		IDetailCategoryBuilder* ClothPartCategoryBuilder = nullptr;
 		IDetailCategoryBuilder* NearestNeighborCategoryBuilder = nullptr;
 		IDetailCategoryBuilder* KMeansCategoryBuilder = nullptr;
