Fix nearest neighbor editor crash for all zero deltas

diff --git a/Engine/Plugins/Animation/MLDeformer/NearestNeighborModel/Source/NearestNeighborModel/Private/NearestNeighborModel.cpp b/Engine/Plugins/Animation/MLDeformer/NearestNeighborModel/Source/NearestNeighborModel/Private/NearestNeighborModel.cpp
--- a/Engine/Plugins/Animation/MLDeformer/NearestNeighborModel/Source/NearestNeighborModel/Private/NearestNeighborModel.cpp
+++ b/Engine/Plugins/Animation/MLDeformer/NearestNeighborModel/Source/NearestNeighborModel/Private/NearestNeighborModel.cpp
@@ -95,6 +95,22 @@ void UNearestNeighborModel::ClipInputs(float* InputPtr, int NumInputs)
 	}	
 }
 
+class FMorphTargetBuffersInfo : public FMorphTargetVertexInfoBuffers
+{
+public:
+	uint32 GetMorphDataSize() const { return MorphData.Num() * sizeof(uint32); }
+	void ResetCPUData() { FMorphTargetVertexInfoBuffers::ResetCPUData(); }
+};
+
+void UNearestNeighborModel::ResetMorphBuffers()
+{
+	if (GetMorphTargetSet())
+	{
+		FMorphTargetBuffersInfo* MorphBuffersInfo = static_cast<FMorphTargetBuffersInfo*>(&GetMorphTargetSet()->MorphBuffers);
+		MorphBuffersInfo->ResetCPUData();
+	}
+}
+
 #if WITH_EDITORONLY_DATA
 void UNearestNeighborModel::UpdatePCACoeffNums()
 {
@@ -228,12 +244,6 @@ void UNearestNeighborModel::UpdateNetworkSize()
 	}
 }
 
-class FMorphTargetBuffersInfo : public FMorphTargetVertexInfoBuffers
-{
-public:
-	uint32 GetMorphDataSize() const { return MorphData.Num() * sizeof(uint32); }
-};
-
 void UNearestNeighborModel::UpdateMorphTargetSize()
 {
 	if (GetMorphTargetSet())
diff --git a/Engine/Plugins/Animation/MLDeformer/NearestNeighborModel/Source/NearestNeighborModel/Public/NearestNeighborModel.h b/Engine/Plugins/Animation/MLDeformer/NearestNeighborModel/Source/NearestNeighborModel/Public/NearestNeighborModel.h
--- a/Engine/Plugins/Animation/MLDeformer/NearestNeighborModel/Source/NearestNeighborModel/Public/NearestNeighborModel.h
+++ b/Engine/Plugins/Animation/MLDeformer/NearestNeighborModel/Source/NearestNeighborModel/Public/NearestNeighborModel.h
@@ -206,6 +206,8 @@ public:
 	void InitInputInfo();
 	void InitPreviousWeights();
 
+	void ResetMorphBuffers();
+
 #if WITH_EDITORONLY_DATA
 	TObjectPtr<UAnimSequence> GetNearestNeighborSkeletons(int32 PartId);
 	const TObjectPtr<UAnimSequence> GetNearestNeighborSkeletons(int32 PartId) const;
diff --git a/Engine/Plugins/Animation/MLDeformer/NearestNeighborModel/Source/NearestNeighborModelEditor/Private/NearestNeighborEditorModel.cpp b/Engine/Plugins/Animation/MLDeformer/NearestNeighborModel/Source/NearestNeighborModelEditor/Private/NearestNeighborEditorModel.cpp
--- a/Engine/Plugins/Animation/MLDeformer/NearestNeighborModel/Source/NearestNeighborModelEditor/Private/NearestNeighborEditorModel.cpp
+++ b/Engine/Plugins/Animation/MLDeformer/NearestNeighborModel/Source/NearestNeighborModelEditor/Private/NearestNeighborEditorModel.cpp
@@ -417,6 +417,12 @@ namespace UE::NearestNeighborModel
 		check(NearestNeighborModel->GetMorphTargetSet().IsValid());
 		FMorphTargetVertexInfoBuffers& MorphBuffers = NearestNeighborModel->GetMorphTargetSet()->MorphBuffers;
 		CompressEngineMorphTargets(MorphBuffers, MorphTargets, LOD, NearestNeighborModel->GetMorphTargetErrorTolerance());
+
+		if (MorphBuffers.GetNumBatches() <= 0)
+		{
+			NearestNeighborModel->ResetMorphBuffers();
+		}
+
 		NearestNeighborModel->SetMorphTargetDeltas(Deltas);
 	}
 
@@ -427,7 +433,7 @@ namespace UE::NearestNeighborModel
 		{
 			FMorphTargetVertexInfoBuffers& MorphBuffers = GetNearestNeighborModel()->GetMorphTargetSet()->MorphBuffers;
 			BeginReleaseResource(&MorphBuffers);
-			if (MorphBuffers.IsMorphCPUDataValid() && MorphBuffers.GetNumMorphs() > 0)
+			if (MorphBuffers.IsMorphCPUDataValid() && MorphBuffers.GetNumMorphs() > 0 && MorphBuffers.GetNumBatches() > 0)
 			{
 				BeginInitResource(&MorphBuffers);
 			}
