Fix crash with empty vertex map and add vertex counts info

diff --git a/Engine/Plugins/Animation/MLDeformer/NearestNeighborModel/Source/NearestNeighborModelEditor/Private/NearestNeighborEditorModel.cpp b/Engine/Plugins/Animation/MLDeformer/NearestNeighborModel/Source/NearestNeighborModelEditor/Private/NearestNeighborEditorModel.cpp
--- a/Engine/Plugins/Animation/MLDeformer/NearestNeighborModel/Source/NearestNeighborModelEditor/Private/NearestNeighborEditorModel.cpp
+++ b/Engine/Plugins/Animation/MLDeformer/NearestNeighborModel/Source/NearestNeighborModelEditor/Private/NearestNeighborEditorModel.cpp
@@ -358,6 +358,10 @@ namespace UE::NearestNeighborModel
 		const int32 NumBaseMeshVerts = Model->GetNumBaseMeshVerts();
 		const UNearestNeighborModel* NearestNeighborModel = GetNearestNeighborModel();
 		const int32 PartNumVerts = VertexMap.Num();
+		if (PartNumVerts == 0)
+		{
+			return;
+		}
 		const int32 NumShapes = FloatArr.Num() / (PartNumVerts * 3);
 		if (DeltaArrayOffset < 0)
 		{
diff --git a/Engine/Plugins/Animation/MLDeformer/NearestNeighborModel/Source/NearestNeighborModelEditor/Private/NearestNeighborModelDetails.cpp b/Engine/Plugins/Animation/MLDeformer/NearestNeighborModel/Source/NearestNeighborModelEditor/Private/NearestNeighborModelDetails.cpp
--- a/Engine/Plugins/Animation/MLDeformer/NearestNeighborModel/Source/NearestNeighborModelEditor/Private/NearestNeighborModelDetails.cpp
+++ b/Engine/Plugins/Animation/MLDeformer/NearestNeighborModel/Source/NearestNeighborModelEditor/Private/NearestNeighborModelDetails.cpp
@@ -82,6 +82,31 @@ namespace UE::NearestNeighborModel
 
 		// Cloth part settings
 		ClothPartCategoryBuilder->AddProperty(UNearestNeighborModel::GetClothPartEditorDataPropertyName());
+
+		FString VertexCountStr = "VertexCounts:";
+		if (NearestNeighborModel)
+		{
+			const int32 NumParts = NearestNeighborModel->GetNumParts();
+			VertexCountStr += "[";
+			for (int32 PartId = 0; PartId < NumParts; PartId++)
+			{
+				VertexCountStr += FString::Printf(TEXT("%d"), NearestNeighborModel->GetPartNumVerts(PartId));
+				if (PartId != NumParts - 1)
+				{
+					VertexCountStr += ",";
+				}
+			}
+			VertexCountStr += "]";
+		}
+		FText VertexCountText = FText::FromString(VertexCountStr);
+		ClothPartCategoryBuilder->AddCustomRow(FText::FromString(""))
+			.WholeRowContent()
+			[
+				SNew(STextBlock)
+				.Text(VertexCountText)
+			];
+
+
 		FText ButtonText = (NearestNeighborModel && NearestNeighborModel->IsClothPartDataValid()) ? LOCTEXT("Update", "Update") : LOCTEXT("Update *", "Update *");
 		ClothPartCategoryBuilder->AddCustomRow(FText::FromString(""))
 			.WholeRowContent()
