Fix crash with more than one meshes

diff --git a/Engine/Plugins/Animation/MLDeformer/NearestNeighborModel/Source/NearestNeighborModelEditor/Private/NearestNeighborGeomCacheSampler.cpp b/Engine/Plugins/Animation/MLDeformer/NearestNeighborModel/Source/NearestNeighborModelEditor/Private/NearestNeighborGeomCacheSampler.cpp
--- a/Engine/Plugins/Animation/MLDeformer/NearestNeighborModel/Source/NearestNeighborModelEditor/Private/NearestNeighborGeomCacheSampler.cpp
+++ b/Engine/Plugins/Animation/MLDeformer/NearestNeighborModel/Source/NearestNeighborModelEditor/Private/NearestNeighborGeomCacheSampler.cpp
@@ -28,15 +28,15 @@ namespace UE::NearestNeighborModel
 			const int32 LODIndex = 0;
 			const FSkeletalMeshLODModel& LODModel = ImportedModel->LODModels[LODIndex];
 			const TArray<FSkelMeshImportedMeshInfo>& SkelMeshInfos = LODModel.ImportedMeshInfos;
-			check(SkelMeshInfos.Num() == 1);
+			check(SkelMeshInfos.Num() >= 1);
 			const FSkelMeshImportedMeshInfo& MeshInfo = SkelMeshInfos[0]; 
 			check(MeshInfo.StartImportedVertex == 0);
-			check(GeometryCache->Tracks.Num() == 1);
+			check(GeometryCache->Tracks.Num() >= 1);
 			UGeometryCacheTrack* Track = GeometryCache->Tracks[0];
 			GeomCacheMeshDatas.Reset(1);
 			GeomCacheMeshDatas.AddDefaulted(1);
 			FGeometryCacheMeshData& GeomCacheMeshData = GeomCacheMeshDatas[0];
-			check(MeshMappings.Num() == 1);
+			check(MeshMappings.Num() >= 1);
 			const UE::MLDeformer::FMLDeformerGeomCacheMeshMapping& MeshMapping = MeshMappings[0]; 
 
 			if (!Track->GetMeshDataAtTime(SampleTime, GeomCacheMeshData))
@@ -94,10 +94,10 @@ namespace UE::NearestNeighborModel
 		{
 			FSkeletalMeshModel* ImportedModel = SkeletalMesh->GetImportedModel();
 			check(ImportedModel);
-			check(ImportedModel->LODModels[0].ImportedMeshInfos.Num() == 1);
+			check(ImportedModel->LODModels[0].ImportedMeshInfos.Num() >= 1);
 			check(ImportedModel->LODModels[0].ImportedMeshInfos[0].StartImportedVertex == 0);
 
-			check(GeometryCache->Tracks.Num() == 1);
+			check(GeometryCache->Tracks.Num() >= 1);
 			UGeometryCacheTrack* Track = GeometryCache->Tracks[0];
 
 			MeshMappings.Reset();
