NearestNeighborEditor: map skeletal mesh and geometry cache properly
add button to use part only mesh

diff --git a/Engine/Plugins/Animation/MLDeformer/NearestNeighborModel/Source/NearestNeighborModel/Public/NearestNeighborModel.h b/Engine/Plugins/Animation/MLDeformer/NearestNeighborModel/Source/NearestNeighborModel/Public/NearestNeighborModel.h
--- a/Engine/Plugins/Animation/MLDeformer/NearestNeighborModel/Source/NearestNeighborModel/Public/NearestNeighborModel.h
+++ b/Engine/Plugins/Animation/MLDeformer/NearestNeighborModel/Source/NearestNeighborModel/Public/NearestNeighborModel.h
@@ -274,6 +274,10 @@ public:
 	static FName GetRecomputeDeltasPropertyName() { return GET_MEMBER_NAME_CHECKED(UNearestNeighborModel, bRecomputeDeltas); }
 	static FName GetRecomputePCAPropertyName() { return GET_MEMBER_NAME_CHECKED(UNearestNeighborModel, bRecomputePCA); }
 
+
+	bool GetUsePartOnlyMesh() const { return bUsePartOnlyMesh; }
+	static FName GetUsePartOnlyMeshPropertyName() { return GET_MEMBER_NAME_CHECKED(UNearestNeighborModel, bUsePartOnlyMesh); }
+
 	UFUNCTION(BlueprintPure, Category = "Nearest Neighbor Model")
 	FString GetModelDir() const;
 #endif
@@ -308,6 +312,9 @@ protected:
 	UPROPERTY()
 	bool bNearestNeighborDataValid = false;
 
+	UPROPERTY(EditAnywhere, BlueprintReadWrite, Category = "Nearest Neighbors")
+	bool bUsePartOnlyMesh = true;
+
 	UPROPERTY(EditAnywhere, BlueprintReadWrite, Category = "Nearest Neighbors")
 	TArray<FSkeletonCachePair> NearestNeighborData;
 
diff --git a/Engine/Plugins/Animation/MLDeformer/NearestNeighborModel/Source/NearestNeighborModelEditor/Private/NearestNeighborEditorModel.cpp b/Engine/Plugins/Animation/MLDeformer/NearestNeighborModel/Source/NearestNeighborModelEditor/Private/NearestNeighborEditorModel.cpp
--- a/Engine/Plugins/Animation/MLDeformer/NearestNeighborModel/Source/NearestNeighborModelEditor/Private/NearestNeighborEditorModel.cpp
+++ b/Engine/Plugins/Animation/MLDeformer/NearestNeighborModel/Source/NearestNeighborModelEditor/Private/NearestNeighborEditorModel.cpp
@@ -204,7 +204,7 @@ namespace UE::NearestNeighborModel
 					GeometryCacheComponent->SetManualTick(true);
 					GeometryCacheComponent->SetPlaybackSpeed(1.0f);
 					GeometryCacheComponent->Play();
-					GeomCacheSampler->GeneratePartMeshMappings(GetNearestNeighborModel()->PartVertexMap(PartId));
+					GeomCacheSampler->GeneratePartMeshMappings(GetNearestNeighborModel()->PartVertexMap(PartId), NearestNeighborModel->GetUsePartOnlyMesh());
 
 					return true;
 				}
diff --git a/Engine/Plugins/Animation/MLDeformer/NearestNeighborModel/Source/NearestNeighborModelEditor/Private/NearestNeighborGeomCacheSampler.cpp b/Engine/Plugins/Animation/MLDeformer/NearestNeighborModel/Source/NearestNeighborModelEditor/Private/NearestNeighborGeomCacheSampler.cpp
--- a/Engine/Plugins/Animation/MLDeformer/NearestNeighborModel/Source/NearestNeighborModelEditor/Private/NearestNeighborGeomCacheSampler.cpp
+++ b/Engine/Plugins/Animation/MLDeformer/NearestNeighborModel/Source/NearestNeighborModelEditor/Private/NearestNeighborGeomCacheSampler.cpp
@@ -28,16 +28,19 @@ namespace UE::NearestNeighborModel
 			const int32 LODIndex = 0;
 			const FSkeletalMeshLODModel& LODModel = ImportedModel->LODModels[LODIndex];
 			const TArray<FSkelMeshImportedMeshInfo>& SkelMeshInfos = LODModel.ImportedMeshInfos;
-			check(SkelMeshInfos.Num() >= 1);
-			const FSkelMeshImportedMeshInfo& MeshInfo = SkelMeshInfos[0]; 
+
+			check(MeshMappings.Num() >= 1);
+			const UE::MLDeformer::FMLDeformerGeomCacheMeshMapping& MeshMapping = MeshMappings[0]; 
+
+			check(SkelMeshInfos.Num() > MeshMapping.MeshIndex);
+			const FSkelMeshImportedMeshInfo& MeshInfo = SkelMeshInfos[MeshMapping.MeshIndex]; 
 			check(MeshInfo.StartImportedVertex == 0);
-			check(GeometryCache->Tracks.Num() >= 1);
-			UGeometryCacheTrack* Track = GeometryCache->Tracks[0];
+
+			check(GeometryCache->Tracks.Num() > MeshMapping.TrackIndex);
+			UGeometryCacheTrack* Track = GeometryCache->Tracks[MeshMapping.TrackIndex];
 			GeomCacheMeshDatas.Reset(1);
 			GeomCacheMeshDatas.AddDefaulted(1);
 			FGeometryCacheMeshData& GeomCacheMeshData = GeomCacheMeshDatas[0];
-			check(MeshMappings.Num() >= 1);
-			const UE::MLDeformer::FMLDeformerGeomCacheMeshMapping& MeshMapping = MeshMappings[0]; 
 
 			if (!Track->GetMeshDataAtTime(SampleTime, GeomCacheMeshData))
 			{
@@ -55,7 +58,6 @@ namespace UE::NearestNeighborModel
 			for(int32 PartVertexIndex = 0; PartVertexIndex < NumPartVerts; PartVertexIndex++)
 			{
 				const int32 VertexIndex = VertexMap[PartVertexIndex];
-				const int32 SkinnedVertexIndex = MeshInfo.StartImportedVertex + VertexIndex;
 				const int32 GeomCacheVertexIndex = MeshMapping.SkelMeshToTrackVertexMap[PartVertexIndex];
 
 				if (GeomCacheVertexIndex != INDEX_NONE && GeomCacheMeshData.Positions.IsValidIndex(GeomCacheVertexIndex))
@@ -84,56 +86,103 @@ namespace UE::NearestNeighborModel
 		}
 	}
 
+	bool IsPotentialMatch(const FString& TrackName, const FString& MeshName)
+	{
+		return (TrackName.Find(MeshName) == 0);
+	}
 
-	void FNearestNeighborGeomCacheSampler::GeneratePartMeshMappings(const TArray<uint32>& VertexMap)
+	void FNearestNeighborGeomCacheSampler::GeneratePartMeshMappings(const TArray<uint32>& VertexMap, bool bUsePartOnlyMesh)
 	{
 		USkeletalMesh* SkeletalMesh = SkeletalMeshComponent.Get() ? SkeletalMeshComponent->GetSkeletalMeshAsset() : nullptr;
 		UGeometryCache* GeometryCache = GeometryCacheComponent.Get() ? GeometryCacheComponent->GetGeometryCache() : nullptr;
 		// TODO: make this more general
 		if (SkeletalMeshComponent && SkeletalMesh && GeometryCacheComponent && GeometryCache)
 		{
+			if (!bUsePartOnlyMesh)
+			{
+				TArray<FString> FailedNames;
+				TArray<FString> VertexMisMatchNames;
+				GenerateGeomCacheMeshMappings(SkeletalMesh, GeometryCache, MeshMappings, FailedNames, VertexMisMatchNames);
+				return;
+			}
 			FSkeletalMeshModel* ImportedModel = SkeletalMesh->GetImportedModel();
 			check(ImportedModel);
 			check(ImportedModel->LODModels[0].ImportedMeshInfos.Num() >= 1);
-			check(ImportedModel->LODModels[0].ImportedMeshInfos[0].StartImportedVertex == 0);
-
-			check(GeometryCache->Tracks.Num() >= 1);
-			UGeometryCacheTrack* Track = GeometryCache->Tracks[0];
 
-			MeshMappings.Reset();
-			MeshMappings.AddDefaulted();
-			UE::MLDeformer::FMLDeformerGeomCacheMeshMapping& Mapping = MeshMappings.Last();
-			Mapping.MeshIndex = 0;
-			Mapping.TrackIndex = 0;
-
-			FGeometryCacheMeshData GeomCacheMeshData;
-			if (!Track->GetMeshDataAtTime(0.0f/*SampleTime*/, GeomCacheMeshData))
+			const TArray<FSkelMeshImportedMeshInfo>& SkelMeshInfos = ImportedModel->LODModels[0].ImportedMeshInfos;
+			if (SkelMeshInfos.IsEmpty())
 			{
 				return;
 			}
 
-			const int32 NumPartVerts = VertexMap.Num();
-			if (GeomCacheMeshData.ImportedVertexNumbers.IsEmpty())
-			{
-				UE_LOG(LogNearestNeighborModel, Error, TEXT("Geometry cache has no stored vertex numbers"));
-				return;
-			}
-			const int32 MaxVertexNumber = FMath::Max(GeomCacheMeshData.ImportedVertexNumbers);
-			if (NumPartVerts != MaxVertexNumber + 1)
-			{
-				UE_LOG(LogNearestNeighborModel, Error, TEXT("Vertex number mismatch: part has %d vertices but geometry cache has %d vertices"), NumPartVerts, MaxVertexNumber + 1);
-			}
+			MeshMappings.Reset();
 
-			Mapping.SkelMeshToTrackVertexMap.AddUninitialized(NumPartVerts);
-			Mapping.ImportedVertexToRenderVertexMap.AddUninitialized(NumPartVerts);
+			FString SkelMeshName;
+			const bool bIsSoloMesh = (GeometryCache->Tracks.Num() == 1 && SkelMeshInfos.Num() == 1);	// Do we just have one mesh and one track?
 
-			for(int32 PartVertexIndex = 0; PartVertexIndex < NumPartVerts; PartVertexIndex++)
+			for (int32 TrackIndex = 0; TrackIndex < GeometryCache->Tracks.Num(); ++TrackIndex)
 			{
-				Mapping.SkelMeshToTrackVertexMap[PartVertexIndex] = GeomCacheMeshData.ImportedVertexNumbers.Find(PartVertexIndex);
-				const int32 VertexIndex = VertexMap[PartVertexIndex];
-				const int32 RenderVertexIndex = ImportedModel->LODModels[0].MeshToImportVertexMap.Find(VertexIndex);
-				Mapping.ImportedVertexToRenderVertexMap[PartVertexIndex] = RenderVertexIndex;
-			}
+				// Check if this is a candidate based on the mesh and track name.
+				UGeometryCacheTrack* Track = GeometryCache->Tracks[TrackIndex];
+
+				bool bFoundMatch = false;
+				for (int32 SkelMeshIndex = 0; SkelMeshIndex < SkelMeshInfos.Num(); ++SkelMeshIndex)
+				{
+					const FSkelMeshImportedMeshInfo& MeshInfo = SkelMeshInfos[SkelMeshIndex];
+
+					// Hack for now
+					if (MeshInfo.StartImportedVertex != 0)
+					{
+						continue;
+					}
+
+					SkelMeshName = MeshInfo.Name.ToString();
+					if (Track &&
+						(IsPotentialMatch(Track->GetName(), SkelMeshName) || bIsSoloMesh))
+					{
+						// Extract the geom cache mesh data.
+						FGeometryCacheMeshData GeomCacheMeshData;
+						if (!Track->GetMeshDataAtTime(SampleTime, GeomCacheMeshData))
+						{
+							continue;
+						}
+
+						// Verify that we have imported vertex numbers.
+						if (GeomCacheMeshData.ImportedVertexNumbers.IsEmpty())
+						{
+							continue;
+						}
+
+						// Create a new mesh mapping entry.
+						MeshMappings.AddDefaulted();
+						UE::MLDeformer::FMLDeformerGeomCacheMeshMapping& Mapping = MeshMappings.Last();
+						Mapping.MeshIndex = SkelMeshIndex;
+						Mapping.TrackIndex = TrackIndex;
+
+						const int32 NumPartVerts = VertexMap.Num();
+						Mapping.SkelMeshToTrackVertexMap.AddUninitialized(NumPartVerts);
+						Mapping.ImportedVertexToRenderVertexMap.AddUninitialized(NumPartVerts);
+
+
+						for(int32 PartVertexIndex = 0; PartVertexIndex < NumPartVerts; PartVertexIndex++)
+						{
+							Mapping.SkelMeshToTrackVertexMap[PartVertexIndex] = GeomCacheMeshData.ImportedVertexNumbers.Find(PartVertexIndex);
+							const int32 VertexIndex = VertexMap[PartVertexIndex];
+							const int32 RenderVertexIndex = ImportedModel->LODModels[0].MeshToImportVertexMap.Find(VertexIndex);
+							Mapping.ImportedVertexToRenderVertexMap[PartVertexIndex] = RenderVertexIndex;
+						}
+
+						// We found a match, no need to iterate over more Tracks.
+						bFoundMatch = true;
+						break;
+					} // If the track name matches the skeletal meshes internal mesh name.
+				} // For all meshes in the Skeletal Mesh.
+
+				if (Track && !bFoundMatch)
+				{
+					UE_LOG(LogMLDeformer, Warning, TEXT("Geometry cache '%s' cannot be matched with a mesh inside the Skeletal Mesh."), *Track->GetName());
+				}
+			} // For all tracks.
 		}
 	}
 
diff --git a/Engine/Plugins/Animation/MLDeformer/NearestNeighborModel/Source/NearestNeighborModelEditor/Private/NearestNeighborModelDetails.cpp b/Engine/Plugins/Animation/MLDeformer/NearestNeighborModel/Source/NearestNeighborModelEditor/Private/NearestNeighborModelDetails.cpp
--- a/Engine/Plugins/Animation/MLDeformer/NearestNeighborModel/Source/NearestNeighborModelEditor/Private/NearestNeighborModelDetails.cpp
+++ b/Engine/Plugins/Animation/MLDeformer/NearestNeighborModel/Source/NearestNeighborModelEditor/Private/NearestNeighborModelDetails.cpp
@@ -133,6 +133,7 @@ namespace UE::NearestNeighborModel
 		// Nearest Neighbor settings
 		NearestNeighborCategoryBuilder->AddProperty(GET_MEMBER_NAME_STRING_CHECKED(UNearestNeighborModel, DecayFactor));
 		NearestNeighborCategoryBuilder->AddProperty(GET_MEMBER_NAME_STRING_CHECKED(UNearestNeighborModel, NearestNeighborOffsetWeight));
+		NearestNeighborCategoryBuilder->AddProperty(UNearestNeighborModel::GetUsePartOnlyMeshPropertyName());
 		NearestNeighborCategoryBuilder->AddProperty(UNearestNeighborModel::GetNearestNeighborDataPropertyName());
 		ButtonText = NearestNeighborModel->IsNearestNeighborDataValid() ? LOCTEXT("Update", "Update") : LOCTEXT("Update *", "Update *");
 
diff --git a/Engine/Plugins/Animation/MLDeformer/NearestNeighborModel/Source/NearestNeighborModelEditor/Public/NearestNeighborGeomCacheSampler.h b/Engine/Plugins/Animation/MLDeformer/NearestNeighborModel/Source/NearestNeighborModelEditor/Public/NearestNeighborGeomCacheSampler.h
--- a/Engine/Plugins/Animation/MLDeformer/NearestNeighborModel/Source/NearestNeighborModelEditor/Public/NearestNeighborGeomCacheSampler.h
+++ b/Engine/Plugins/Animation/MLDeformer/NearestNeighborModel/Source/NearestNeighborModelEditor/Public/NearestNeighborGeomCacheSampler.h
@@ -14,7 +14,7 @@ namespace UE::NearestNeighborModel
 		virtual void SampleKMeansAnim(const int32 SkeletonId);
 		virtual void SampleKMeansFrame(const int32 Frame);
 		const TArray<float>& GetPartVertexDeltas() const { return PartVertexDeltas; }
-		void GeneratePartMeshMappings(const TArray<uint32>& VertexMap);
+		void GeneratePartMeshMappings(const TArray<uint32>& VertexMap, bool bUsePartOnlyMesh);
 	protected:
 		TArray<float> PartVertexDeltas;
 	};
