NearestNeighborModel: Add error messages for morph target

diff --git a/Engine/Plugins/Animation/MLDeformer/NearestNeighborModel/Source/NearestNeighborModel/Public/NearestNeighborModel.h b/Engine/Plugins/Animation/MLDeformer/NearestNeighborModel/Source/NearestNeighborModel/Public/NearestNeighborModel.h
--- a/Engine/Plugins/Animation/MLDeformer/NearestNeighborModel/Source/NearestNeighborModel/Public/NearestNeighborModel.h
+++ b/Engine/Plugins/Animation/MLDeformer/NearestNeighborModel/Source/NearestNeighborModel/Public/NearestNeighborModel.h
@@ -115,6 +115,13 @@ struct FClothPartData
 	TArray<float> NeighborOffsets;
 };
 
+enum EResultMessage : uint8
+{
+	SUCCESS = 0,
+	ERROR = 1,
+	WARNING = 2
+};
+
 /**
  * The nearest neighbor model.
  * This model contains the PCA basis of the cloth vertex deltas and a small set of cloth for nearest neighbor search. 
diff --git a/Engine/Plugins/Animation/MLDeformer/NearestNeighborModel/Source/NearestNeighborModelEditor/Private/NearestNeighborEditorModel.cpp b/Engine/Plugins/Animation/MLDeformer/NearestNeighborModel/Source/NearestNeighborModelEditor/Private/NearestNeighborEditorModel.cpp
--- a/Engine/Plugins/Animation/MLDeformer/NearestNeighborModel/Source/NearestNeighborModelEditor/Private/NearestNeighborEditorModel.cpp
+++ b/Engine/Plugins/Animation/MLDeformer/NearestNeighborModel/Source/NearestNeighborModelEditor/Private/NearestNeighborEditorModel.cpp
@@ -32,7 +32,6 @@ namespace UE::NearestNeighborModel
 	{
 		FMLDeformerEditorModel::Init(InitSettings);
 		InitInputInfo(Model->GetInputInfo());
-		InitMorphTargets();
 	}
 
 	FMLDeformerSampler* FNearestNeighborEditorModel::CreateSampler() const
@@ -379,32 +378,81 @@ namespace UE::NearestNeighborModel
 			}
 		}
 	}
+#define RETURN_ON_ERROR(Flag, Func)\
+	Flag |= Func;\
+	if ((Flag & ERROR) != 0)\
+	{\
+		GetEditor()->GetModelDetailsView()->ForceRefresh();\
+		return;\
+	}
+	
+	void FNearestNeighborEditorModel::OnMorphTargetUpdate()
+	{
+		MorphTargetUpdateResult = SUCCESS;
+		RETURN_ON_ERROR(MorphTargetUpdateResult, InitMorphTargets());
+		RefreshMorphTargets();
+		GetNearestNeighborModel()->UpdateNetworkSize();
+		GetNearestNeighborModel()->UpdateMorphTargetSize();
+
+		GetEditor()->GetModelDetailsView()->ForceRefresh();
+	}
 
-	void FNearestNeighborEditorModel::InitMorphTargets()
+	uint8 FNearestNeighborEditorModel::InitMorphTargets()
 	{
+		EResultMessage Result = SUCCESS;
 		UNearestNeighborModel* NearestNeighborModel = GetNearestNeighborModel();
-		if (NearestNeighborModel->GetNumBaseMeshVerts() == 0)
+
+		const USkeletalMesh* SkelMesh = Model->GetSkeletalMesh();
+		const int32 NumBaseMeshVerts = Model->GetNumBaseMeshVerts();
+		if (!SkelMesh || NumBaseMeshVerts == 0)
 		{
-			return;
+			UE_LOG(LogNearestNeighborModel, Error, TEXT("SkeletalMesh is empty. No morph targets are generated"));
+			return ERROR;
+		}
+
+		if (Model->GetVertexMap().IsEmpty())
+		{
+			UE_LOG(LogNearestNeighborModel, Error, TEXT("VertexMap of the skeletal mesh is empty. No morph targets are generated"));
+			return ERROR;
+		}
+
+		const int32 NumParts = NearestNeighborModel->GetNumParts();
+		if (NumParts == 0)
+		{
+			UE_LOG(LogNearestNeighborModel, Error, TEXT("There are no cloth parts. No morph targets are generated"));
+			return ERROR;
+		}
+
+		const int32 NumImportedModelVerts = FMath::Max(NearestNeighborModel->GetVertexMap()) + 1;
+		if (NumImportedModelVerts != NumBaseMeshVerts)
+		{
+			UE_LOG(LogNearestNeighborModel, Error, TEXT("Vertex count mismatch: imported model of SkeletalMesh has %d vertices and cached SkeletalMesh has %d vertices"), NumImportedModelVerts, NumBaseMeshVerts);
+			return ERROR;
 		}
+
 		TArray<FVector3f> Deltas;
 		Deltas.Reset();
 
 		int32 NumPCACoeff = 0;
-		for (int32 PartId = 0; PartId < NearestNeighborModel->GetNumParts(); PartId++)
+		for (int32 PartId = 0; PartId < NumParts; PartId++)
 		{
 			NumPCACoeff += NearestNeighborModel->GetPCACoeffNum(PartId);
 		}
-		Deltas.Reserve((1 + NumPCACoeff) * NearestNeighborModel->GetNumBaseMeshVerts());
+		Deltas.Reserve((1 + NumPCACoeff) * NumBaseMeshVerts);
 
-		for (int32 PartId = 0; PartId < NearestNeighborModel->GetNumParts(); PartId++)
+		for (int32 PartId = 0; PartId < NumParts; PartId++)
 		{
 			const TArray<uint32>& VertexMap = NearestNeighborModel->PartVertexMap(PartId);
+			if (VertexMap.IsEmpty())
+			{
+				UE_LOG(LogNearestNeighborModel, Warning, TEXT("Cloth part %d has empty vertex map. No morph targets are generated for this part."), PartId);
+				Result = WARNING;
+			}
 			AddFloatArrayToDeltaArray(NearestNeighborModel->ClothPartData[PartId].VertexMean, VertexMap, Deltas, 0);
 			AddFloatArrayToDeltaArray(NearestNeighborModel->ClothPartData[PartId].PCABasis, VertexMap, Deltas);
 		}
 
-		for (int32 PartId = 0; PartId < NearestNeighborModel->GetNumParts(); PartId++)
+		for (int32 PartId = 0; PartId < NumParts; PartId++)
 		{
 			const TArray<uint32>& VertexMap = NearestNeighborModel->PartVertexMap(PartId);
 			AddFloatArrayToDeltaArray(NearestNeighborModel->ClothPartData[PartId].NeighborOffsets, VertexMap, Deltas);
@@ -412,9 +460,11 @@ namespace UE::NearestNeighborModel
 
 		if (Deltas.Num() == 0)
 		{
-			return;
+			UE_LOG(LogNearestNeighborModel, Error, TEXT("All cloth parts are empty. No morph targets are generated."));
+			return ERROR;
 		}
 
+
 		const int32 LOD = 0;
 		TArray<UMorphTarget*> MorphTargets;
 		CreateEngineMorphTargets(MorphTargets, Deltas, FString("NNMorphTarget_"), LOD, NearestNeighborModel->GetMorphTargetDeltaThreshold());
@@ -424,10 +474,13 @@ namespace UE::NearestNeighborModel
 
 		if (MorphBuffers.GetNumBatches() <= 0)
 		{
+			UE_LOG(LogNearestNeighborModel, Warning, TEXT("Morph buffer is empty. It is possible that all deltas are zero. No morph targets are generated."));
+			Result = WARNING;
 			NearestNeighborModel->ResetMorphBuffers();
 		}
 
 		NearestNeighborModel->SetMorphTargetDeltas(Deltas);
+		return Result;
 	}
 
 	void FNearestNeighborEditorModel::RefreshMorphTargets()
@@ -435,6 +488,7 @@ namespace UE::NearestNeighborModel
 		USkeletalMeshComponent* SkelMeshComponent = FindEditorActor(ActorID_Test_MLDeformed)->GetSkeletalMeshComponent() ;
 		if (SkelMeshComponent)
 		{
+			check(GetNearestNeighborModel()->GetMorphTargetSet().IsValid());
 			FMorphTargetVertexInfoBuffers& MorphBuffers = GetNearestNeighborModel()->GetMorphTargetSet()->MorphBuffers;
 			BeginReleaseResource(&MorphBuffers);
 			if (MorphBuffers.IsMorphCPUDataValid() && MorphBuffers.GetNumMorphs() > 0 && MorphBuffers.GetNumBatches() > 0)
diff --git a/Engine/Plugins/Animation/MLDeformer/NearestNeighborModel/Source/NearestNeighborModelEditor/Private/NearestNeighborModelDetails.cpp b/Engine/Plugins/Animation/MLDeformer/NearestNeighborModel/Source/NearestNeighborModelEditor/Private/NearestNeighborModelDetails.cpp
--- a/Engine/Plugins/Animation/MLDeformer/NearestNeighborModel/Source/NearestNeighborModelEditor/Private/NearestNeighborModelDetails.cpp
+++ b/Engine/Plugins/Animation/MLDeformer/NearestNeighborModel/Source/NearestNeighborModelEditor/Private/NearestNeighborModelDetails.cpp
@@ -116,16 +116,10 @@ namespace UE::NearestNeighborModel
 				.HAlign(HAlign_Center)
 				.OnClicked_Lambda([this]
 				{
-					if (NearestNeighborModel != nullptr)
-					{
-						NearestNeighborModel->UpdateClothPartData();
-						NearestNeighborModel->InitPreviousWeights();
-						if (NearestNeighborEditorModel != nullptr)
-						{
-							NearestNeighborEditorModel->UpdateNearestNeighborActors();
-						}
-						EditorModel->GetEditor()->GetModelDetailsView()->ForceRefresh();
-					}
+					NearestNeighborModel->UpdateClothPartData();
+					NearestNeighborModel->InitPreviousWeights();
+					NearestNeighborEditorModel->UpdateNearestNeighborActors();
+					EditorModel->GetEditor()->GetModelDetailsView()->ForceRefresh();
 					return FReply::Handled();
 				})
 			];
@@ -145,12 +139,9 @@ namespace UE::NearestNeighborModel
 				.HAlign(HAlign_Center)
 				.OnClicked_Lambda([this]
 				{
-					if (NearestNeighborEditorModel != nullptr)
-					{
-						NearestNeighborEditorModel->UpdateNearestNeighborData();
-						NearestNeighborModel->InitPreviousWeights();
-						EditorModel->GetEditor()->GetModelDetailsView()->ForceRefresh();
-					}
+					NearestNeighborEditorModel->UpdateNearestNeighborData();
+					NearestNeighborModel->InitPreviousWeights();
+					EditorModel->GetEditor()->GetModelDetailsView()->ForceRefresh();
 					return FReply::Handled();
 				})
 			];
@@ -169,7 +160,7 @@ namespace UE::NearestNeighborModel
 					return FReply::Handled();
 				})
 			];
-
+		
 		MorphTargetCategoryBuilder->AddProperty(UNearestNeighborModel::GetMorphDataSizePropertyName());
 		MorphTargetCategoryBuilder->AddCustomRow(FText::FromString(""))
 			.WholeRowContent()
@@ -179,18 +170,45 @@ namespace UE::NearestNeighborModel
 				.HAlign(HAlign_Center)
 				.OnClicked_Lambda([this]
 				{
-					if (NearestNeighborEditorModel != nullptr)
-					{
-						NearestNeighborEditorModel->InitMorphTargets();
-						NearestNeighborEditorModel->RefreshMorphTargets();
-						NearestNeighborModel->UpdateNetworkSize();
-						NearestNeighborModel->UpdateMorphTargetSize();
-						EditorModel->GetEditor()->GetModelDetailsView()->ForceRefresh();
-					}
+					NearestNeighborEditorModel->OnMorphTargetUpdate();
 					return FReply::Handled();
 				})
 			];
+		AddUpdateResultText(MorphTargetCategoryBuilder, NearestNeighborEditorModel->GetMorphTargetUpdateResult());
+	}
+
+	void FNearestNeighborModelDetails::AddUpdateResultText(IDetailCategoryBuilder* CategoryBuilder, uint8 Result)
+	{
+		CategoryBuilder->AddCustomRow(FText::FromString("UpdateResultError"))
+			.Visibility((Result & ERROR) != 0 ? EVisibility::Visible : EVisibility::Collapsed)
+			.WholeRowContent()
+			[
+				SNew(SBox)
+				.Padding(FMargin(0.0f, 4.0f))
+				[
+					SNew(SWarningOrErrorBox)
+					.MessageStyle(EMessageStyle::Error)
+					.Message(FText::FromString(TEXT("Update failed with errors. Please check Output Log (LogNearestNeighborModel, LogPython) for details.")))
+				]
+			];
+		CategoryBuilder->AddCustomRow(FText::FromString("UpdateResultWarning"))
+			.Visibility((Result & WARNING) != 0 ? EVisibility::Visible : EVisibility::Collapsed)
+			.WholeRowContent()
+			[
+				SNew(SBox)
+				.Padding(FMargin(0.0f, 4.0f))
+				[
+					SNew(SWarningOrErrorBox)
+					.MessageStyle(EMessageStyle::Warning)
+					.Message(FText::FromString(TEXT("Update finished with warnings. Please check Output Log (LogNearestNeighborModel, LogPython) for details.")))
+				]
+			];
+		if (Result == SUCCESS)
+		{
+			UE_LOG(LogNearestNeighborModel, Display, TEXT("Update succeeded."));
+		}
 	}
+
 }	// namespace UE::NearestNeighborModel
 
 #undef LOCTEXT_NAMESPACE
diff --git a/Engine/Plugins/Animation/MLDeformer/NearestNeighborModel/Source/NearestNeighborModelEditor/Public/NearestNeighborEditorModel.h b/Engine/Plugins/Animation/MLDeformer/NearestNeighborModel/Source/NearestNeighborModelEditor/Public/NearestNeighborEditorModel.h
--- a/Engine/Plugins/Animation/MLDeformer/NearestNeighborModel/Source/NearestNeighborModelEditor/Public/NearestNeighborEditorModel.h
+++ b/Engine/Plugins/Animation/MLDeformer/NearestNeighborModel/Source/NearestNeighborModelEditor/Public/NearestNeighborEditorModel.h
@@ -59,12 +59,16 @@ namespace UE::NearestNeighborModel
 		void UpdateNearestNeighborActors();
 		void KMeansClusterPoses();
 
-		void InitMorphTargets();
+		uint8 InitMorphTargets();
 		void RefreshMorphTargets();
 		void AddFloatArrayToDeltaArray(const TArray<float>& FloatArr, const TArray<uint32>& VertexMap, TArray<FVector3f>& DeltaArr, int32 DeltaArrayOffset = -1, float ScaleFactor = 1);
 
 		int32 GetNumParts();
 
+		void OnMorphTargetUpdate();
+		uint8 GetMorphTargetUpdateResult() { return MorphTargetUpdateResult; }
+
+
 	protected:
 		virtual void CreateNearestNeighborActors(UWorld* World, int32 StartIndex = 0);
 
@@ -88,5 +92,8 @@ namespace UE::NearestNeighborModel
 
 	private:
 		UWorld* EditorWorld = nullptr;
+		uint8 ClothPartUpdateResult = 0;
+		uint8 NearestNeighborUpdateResult = 0;
+		uint8 MorphTargetUpdateResult = 0;
 	};
 }	// namespace UE::NearestNeighborModel
diff --git a/Engine/Plugins/Animation/MLDeformer/NearestNeighborModel/Source/NearestNeighborModelEditor/Public/NearestNeighborModelDetails.h b/Engine/Plugins/Animation/MLDeformer/NearestNeighborModel/Source/NearestNeighborModelEditor/Public/NearestNeighborModelDetails.h
--- a/Engine/Plugins/Animation/MLDeformer/NearestNeighborModel/Source/NearestNeighborModelEditor/Public/NearestNeighborModelDetails.h
+++ b/Engine/Plugins/Animation/MLDeformer/NearestNeighborModel/Source/NearestNeighborModelEditor/Public/NearestNeighborModelDetails.h
@@ -25,6 +25,7 @@ namespace UE::NearestNeighborModel
 		virtual void CreateCategories() override;
 
 		virtual bool UpdateMemberPointers(const TArray<TWeakObjectPtr<UObject>>& Objects) override;
+		virtual void AddUpdateResultText(IDetailCategoryBuilder* CategoryBuilder, uint8 Result);
 
 	protected:
 		UNearestNeighborModel* NearestNeighborModel = nullptr;
