Fix crash if skeletal mesh is changed

diff --git a/Engine/Plugins/Animation/MLDeformer/NearestNeighborModel/Source/NearestNeighborModelEditor/Private/NearestNeighborEditorModel.cpp b/Engine/Plugins/Animation/MLDeformer/NearestNeighborModel/Source/NearestNeighborModelEditor/Private/NearestNeighborEditorModel.cpp
--- a/Engine/Plugins/Animation/MLDeformer/NearestNeighborModel/Source/NearestNeighborModelEditor/Private/NearestNeighborEditorModel.cpp
+++ b/Engine/Plugins/Animation/MLDeformer/NearestNeighborModel/Source/NearestNeighborModelEditor/Private/NearestNeighborEditorModel.cpp
@@ -464,6 +464,13 @@ namespace UE::NearestNeighborModel
 			return ERROR;
 		}
 
+		const int32 NumImportedModelVerts = FMath::Max(NearestNeighborModel->GetVertexMap()) + 1;
+		const int32 NumBaseMeshVerts = Model->GetNumBaseMeshVerts();
+		if (NumImportedModelVerts != NumBaseMeshVerts)
+		{
+			UE_LOG(LogNearestNeighborModel, Error, TEXT("Vertex count mismatch: imported model of SkeletalMesh has %d vertices and cached SkeletalMesh has %d vertices"), NumImportedModelVerts, NumBaseMeshVerts);
+			return;
+		}
 
 		const int32 LOD = 0;
 		TArray<UMorphTarget*> MorphTargets;
