NearestNeighborEditor: Fix crash when changing PCACoeffNum or Nearest neighbor number.
Enable python to handle unexpected errors.

diff --git a/Engine/Plugins/Animation/MLDeformer/NearestNeighborModel/Source/NearestNeighborModel/Private/NearestNeighborModel.cpp b/Engine/Plugins/Animation/MLDeformer/NearestNeighborModel/Source/NearestNeighborModel/Private/NearestNeighborModel.cpp
--- a/Engine/Plugins/Animation/MLDeformer/NearestNeighborModel/Source/NearestNeighborModel/Private/NearestNeighborModel.cpp
+++ b/Engine/Plugins/Animation/MLDeformer/NearestNeighborModel/Source/NearestNeighborModel/Private/NearestNeighborModel.cpp
@@ -4,12 +4,13 @@
 #include "NearestNeighborModelInputInfo.h"
 #include "NearestNeighborModelInstance.h"
 #include "MLDeformerComponent.h"
-#include "Components/ExternalMorphSet.h"
 #include "Engine/SkeletalMesh.h"
 #include "GeometryCache.h"
 #include "Misc/FileHelper.h"
 #include "Components/SkinnedMeshComponent.h"
+#include "Components/ExternalMorphSet.h"
 #include "Rendering/SkeletalMeshModel.h"
+#include "Rendering/MorphTargetVertexInfoBuffers.h"
 #include "UObject/UObjectGlobals.h"
 #include "NeuralNetwork.h"
 
@@ -96,6 +97,16 @@ void UNearestNeighborModel::ClipInputs(float* InputPtr, int NumInputs)
 	}	
 }
 
+int32 UNearestNeighborModel::GetTotalNumNeighbors() const
+{
+	int32 TotalNumNeighbors = 0;
+	for (int32 PartId = 0; PartId < GetNumParts(); PartId++)
+	{
+		TotalNumNeighbors += GetNumNeighbors(PartId);
+	}
+	return TotalNumNeighbors;
+}
+
 class FMorphTargetBuffersInfo : public FMorphTargetVertexInfoBuffers
 {
 public:
@@ -290,12 +301,7 @@ bool UNearestNeighborModel::CheckPCAData(int32 PartId) const
 
 void UNearestNeighborModel::InitPreviousWeights()
 {
-	int32 NumWeights = 1;
-	for (int32 PartId = 0; PartId < GetNumParts(); PartId++)
-	{
-		NumWeights += GetPCACoeffNum(PartId);
-		NumWeights += GetNumNeighbors(PartId);
-	}
+	const int32 NumWeights = GetMorphTargetSet()->MorphBuffers.GetNumMorphs();
 	PreviousWeights.SetNumZeroed(NumWeights);
 }
 #undef LOCTEXT_NAMESPACE
diff --git a/Engine/Plugins/Animation/MLDeformer/NearestNeighborModel/Source/NearestNeighborModel/Private/NearestNeighborModelInstance.cpp b/Engine/Plugins/Animation/MLDeformer/NearestNeighborModel/Source/NearestNeighborModel/Private/NearestNeighborModelInstance.cpp
--- a/Engine/Plugins/Animation/MLDeformer/NearestNeighborModel/Source/NearestNeighborModel/Private/NearestNeighborModelInstance.cpp
+++ b/Engine/Plugins/Animation/MLDeformer/NearestNeighborModel/Source/NearestNeighborModel/Private/NearestNeighborModelInstance.cpp
@@ -6,7 +6,6 @@
 #include "NeuralNetwork.h"
 #include "Components/ExternalMorphSet.h"
 #include "Components/SkeletalMeshComponent.h"
-
 using namespace UE::MLDeformer;
 using namespace UE::NearestNeighborModel;
 
@@ -104,6 +103,17 @@ void UNearestNeighborModelInstance::Execute(float ModelWeight)
 	RunNearestNeighborModel(ModelWeight);
 }
 
+TArray<float> ExtractArray(const TArray<float>& InArr, int32 Start, int32 End)
+{
+	int32 Num = End - Start;
+	TArray<float> Result; Result.SetNum(Num);
+	for (int32 i = 0; i < Num; i++)
+	{
+		Result[i] = InArr[Start + i];
+	}
+	return Result;
+}
+
 void UNearestNeighborModelInstance::RunNearestNeighborModel(float ModelWeight)
 {
 	UNearestNeighborModel* NearestNeighborModel = Cast<UNearestNeighborModel>(Model);
@@ -130,8 +140,9 @@ void UNearestNeighborModelInstance::RunNearestNeighborModel(float ModelWeight)
 	{
 		const FNeuralTensor& OutputTensor = NeuralNetwork->GetOutputTensorForContext(NeuralNetworkInferenceHandle);
 		const int32 NumNetworkWeights = OutputTensor.Num();
-		const int32 NumMorphTargets = WeightData->Weights.Num();;
-		if (NumMorphTargets >= NumNetworkWeights + 1)
+		const int32 NumMorphTargets = WeightData->Weights.Num();
+		const int32 TotalNumNeighbors = NearestNeighborModel->GetTotalNumNeighbors();
+		if (NumMorphTargets >= NumNetworkWeights + 1 + TotalNumNeighbors)
 		{
 			WeightData->Weights[0] = 1.0f * ModelWeight;
 
diff --git a/Engine/Plugins/Animation/MLDeformer/NearestNeighborModel/Source/NearestNeighborModel/Public/NearestNeighborModel.h b/Engine/Plugins/Animation/MLDeformer/NearestNeighborModel/Source/NearestNeighborModel/Public/NearestNeighborModel.h
--- a/Engine/Plugins/Animation/MLDeformer/NearestNeighborModel/Source/NearestNeighborModel/Public/NearestNeighborModel.h
+++ b/Engine/Plugins/Animation/MLDeformer/NearestNeighborModel/Source/NearestNeighborModel/Public/NearestNeighborModel.h
@@ -165,6 +165,8 @@ public:
 	UFUNCTION(BlueprintPure, Category = "Nearest Neighbor Model")
 	int32 GetNumNeighbors(int32 PartId) const { return ClothPartData[PartId].NumNeighbors; }
 
+	int32 GetTotalNumNeighbors() const;
+
 	float GetDecayFactor() const { return DecayFactor; };
 	float GetNearestNeighborOffsetWeight() const { return NearestNeighborOffsetWeight; }
 
diff --git a/Engine/Plugins/Animation/MLDeformer/NearestNeighborModel/Source/NearestNeighborModelEditor/Private/NearestNeighborEditorModel.cpp b/Engine/Plugins/Animation/MLDeformer/NearestNeighborModel/Source/NearestNeighborModelEditor/Private/NearestNeighborEditorModel.cpp
--- a/Engine/Plugins/Animation/MLDeformer/NearestNeighborModel/Source/NearestNeighborModelEditor/Private/NearestNeighborEditorModel.cpp
+++ b/Engine/Plugins/Animation/MLDeformer/NearestNeighborModel/Source/NearestNeighborModelEditor/Private/NearestNeighborEditorModel.cpp
@@ -73,7 +73,10 @@ namespace UE::NearestNeighborModel
 	void FNearestNeighborEditorModel::OnPostTraining(ETrainingResult TrainingResult, bool bUsePartiallyTrainedWhenAborted)
 	{
 		FMLDeformerMorphModelEditorModel::OnPostTraining(TrainingResult, bUsePartiallyTrainedWhenAborted);
-		GetNearestNeighborModel()->InitPreviousWeights();
+		if (TrainingResult == ETrainingResult::Success || (TrainingResult == ETrainingResult::Aborted && bUsePartiallyTrainedWhenAborted))
+		{
+			GetNearestNeighborModel()->InitPreviousWeights();
+		}
 	}
 
 	FString FNearestNeighborEditorModel::GetTrainedNetworkOnnxFile() const
diff --git a/Engine/Plugins/Animation/MLDeformer/NearestNeighborModel/Source/NearestNeighborModelEditor/Private/NearestNeighborModelDetails.cpp b/Engine/Plugins/Animation/MLDeformer/NearestNeighborModel/Source/NearestNeighborModelEditor/Private/NearestNeighborModelDetails.cpp
--- a/Engine/Plugins/Animation/MLDeformer/NearestNeighborModel/Source/NearestNeighborModelEditor/Private/NearestNeighborModelDetails.cpp
+++ b/Engine/Plugins/Animation/MLDeformer/NearestNeighborModel/Source/NearestNeighborModelEditor/Private/NearestNeighborModelDetails.cpp
@@ -119,7 +119,6 @@ namespace UE::NearestNeighborModel
 				.OnClicked_Lambda([this]
 				{
 					NearestNeighborModel->UpdateClothPartData();
-					NearestNeighborModel->InitPreviousWeights();
 					NearestNeighborEditorModel->UpdateNearestNeighborActors();
 					EditorModel->GetEditor()->GetModelDetailsView()->ForceRefresh();
 					return FReply::Handled();
@@ -142,7 +141,7 @@ namespace UE::NearestNeighborModel
 				.OnClicked_Lambda([this]
 				{
 					NearestNeighborEditorModel->UpdateNearestNeighborData();
-					NearestNeighborModel->InitPreviousWeights();
+					NearestNeighborEditorModel->UpdateNearestNeighborActors();
 					EditorModel->GetEditor()->GetModelDetailsView()->ForceRefresh();
 					return FReply::Handled();
 				})
@@ -173,6 +172,7 @@ namespace UE::NearestNeighborModel
 				.OnClicked_Lambda([this]
 				{
 					NearestNeighborEditorModel->OnMorphTargetUpdate();
+					NearestNeighborModel->InitPreviousWeights();
 					return FReply::Handled();
 				})
 			];
