NearestNeighborEditor: fix nearest neighbor update bug when geomcache is None

diff --git a/Engine/Plugins/Animation/MLDeformer/NearestNeighborModel/Source/NearestNeighborModelEditor/Private/NearestNeighborEditorModel.cpp b/Engine/Plugins/Animation/MLDeformer/NearestNeighborModel/Source/NearestNeighborModelEditor/Private/NearestNeighborEditorModel.cpp
--- a/Engine/Plugins/Animation/MLDeformer/NearestNeighborModel/Source/NearestNeighborModelEditor/Private/NearestNeighborEditorModel.cpp
+++ b/Engine/Plugins/Animation/MLDeformer/NearestNeighborModel/Source/NearestNeighborModelEditor/Private/NearestNeighborEditorModel.cpp
@@ -92,6 +92,18 @@ namespace UE::NearestNeighborModel
 		}
 	}
 
+	int32 FNearestNeighborEditorModel::GetNumTrainingFrames() const
+	{
+		if (NumTrainingFramesOverride >= 0)
+		{
+			return NumTrainingFramesOverride;
+		}
+		else
+		{
+			return FMLDeformerMorphModelEditorModel::GetNumTrainingFrames();
+		}
+	}
+
 	void FNearestNeighborEditorModel::CreateActors(const TSharedRef<IPersonaPreviewScene>& InPersonaPreviewScene)
 	{
 		FMLDeformerMorphModelEditorModel::CreateActors(InPersonaPreviewScene);
@@ -210,6 +222,8 @@ namespace UE::NearestNeighborModel
 					GeometryCacheComponent->Play();
 					GeomCacheSampler->GeneratePartMeshMappings(GetNearestNeighborModel()->PartVertexMap(PartId), NearestNeighborModel->GetUsePartOnlyMesh());
 
+					NumTrainingFramesOverride = NumFrames;
+
 					return true;
 				}
 				else
@@ -256,7 +270,7 @@ namespace UE::NearestNeighborModel
 				SkeletalMeshComponent->Play(false);
 				SkeletalMeshComponent->RefreshBoneTransforms();
 			}
-
+			NumTrainingFramesOverride = -1;
 			GeomCacheSampler->RegisterTargetComponents();
 		}
 	}
diff --git a/Engine/Plugins/Animation/MLDeformer/NearestNeighborModel/Source/NearestNeighborModelEditor/Public/NearestNeighborEditorModel.h b/Engine/Plugins/Animation/MLDeformer/NearestNeighborModel/Source/NearestNeighborModelEditor/Public/NearestNeighborEditorModel.h
--- a/Engine/Plugins/Animation/MLDeformer/NearestNeighborModel/Source/NearestNeighborModelEditor/Public/NearestNeighborEditorModel.h
+++ b/Engine/Plugins/Animation/MLDeformer/NearestNeighborModel/Source/NearestNeighborModelEditor/Public/NearestNeighborEditorModel.h
@@ -42,6 +42,7 @@ namespace UE::NearestNeighborModel
 		virtual void OnPropertyChanged(FPropertyChangedEvent& PropertyChangedEvent) override;
 		virtual void OnPostTraining(ETrainingResult TrainingResult, bool bUsePartiallyTrainedWhenAborted) override;
 		virtual FString GetTrainedNetworkOnnxFile() const override;
+		virtual int32 GetNumTrainingFrames() const override;
 		// ~END FMLDeformerGeomCacheEditorModel overrides.
 
 		// Some helpers that cast to this model's variants of some classes.
@@ -95,5 +96,6 @@ namespace UE::NearestNeighborModel
 		uint8 ClothPartUpdateResult = 0;
 		uint8 NearestNeighborUpdateResult = 0;
 		uint8 MorphTargetUpdateResult = 0;
+		int32 NumTrainingFramesOverride = -1;
 	};
 }	// namespace UE::NearestNeighborModel
