NearestNeighborEditor: Fix interference between PIE and NNEditor
Move PreviousWeights to instance.
#[fyi] matt.stoneham

diff --git a/Engine/Plugins/Animation/MLDeformer/NearestNeighborModel/Source/NearestNeighborModel/Private/NearestNeighborModel.cpp b/Engine/Plugins/Animation/MLDeformer/NearestNeighborModel/Source/NearestNeighborModel/Private/NearestNeighborModel.cpp
--- a/Engine/Plugins/Animation/MLDeformer/NearestNeighborModel/Source/NearestNeighborModel/Private/NearestNeighborModel.cpp
+++ b/Engine/Plugins/Animation/MLDeformer/NearestNeighborModel/Source/NearestNeighborModel/Private/NearestNeighborModel.cpp
@@ -55,7 +55,6 @@ void UNearestNeighborModel::PostLoad()
 	Super::PostLoad();
 
 	InitInputInfo();
-	InitPreviousWeights();
 
 #if WITH_EDITORONLY_DATA
 	UpdateNetworkInputDim();
@@ -297,11 +296,4 @@ bool UNearestNeighborModel::CheckPCAData(int32 PartId) const
 	const FClothPartData& Data = ClothPartData[PartId];
 	return Data.VertexMap.Num() > 0 && Data.PCABasis.Num() == Data.VertexMap.Num() * 3 * Data.PCACoeffNum;
 }
-
-
-void UNearestNeighborModel::InitPreviousWeights()
-{
-	const int32 NumWeights = GetMorphTargetSet()->MorphBuffers.GetNumMorphs();
-	PreviousWeights.SetNumZeroed(NumWeights);
-}
 #undef LOCTEXT_NAMESPACE
diff --git a/Engine/Plugins/Animation/MLDeformer/NearestNeighborModel/Source/NearestNeighborModel/Private/NearestNeighborModelInstance.cpp b/Engine/Plugins/Animation/MLDeformer/NearestNeighborModel/Source/NearestNeighborModel/Private/NearestNeighborModelInstance.cpp
--- a/Engine/Plugins/Animation/MLDeformer/NearestNeighborModel/Source/NearestNeighborModel/Private/NearestNeighborModelInstance.cpp
+++ b/Engine/Plugins/Animation/MLDeformer/NearestNeighborModel/Source/NearestNeighborModel/Private/NearestNeighborModelInstance.cpp
@@ -34,48 +34,6 @@ int64 UNearestNeighborModelInstance::SetBoneTransforms(float* OutputBuffer, int6
 	return EndIndex;
 }
 
-BEGIN_SHADER_PARAMETER_STRUCT(FCopyBufferParameters, )
-    RDG_BUFFER_ACCESS(Source,  ERHIAccess::CopySrc)
-END_SHADER_PARAMETER_STRUCT()
-
-template<class T, class TRDGBuffer>
-void AddCopyToCPUPass(FRDGBuilder& GraphBuilder, TRDGBuffer RDGBuffer, TArray<T>& Array)
-{
-    FCopyBufferParameters* Parameters = GraphBuilder.AllocParameters<FCopyBufferParameters>();
-    Parameters->Source = RDGBuffer;
-	GraphBuilder.AddPass(
-		RDG_EVENT_NAME("CopyToCPU"),
-        Parameters,
-        ERDGPassFlags::Readback,
-		[RDGBuffer, &Array](FRHICommandListImmediate& CmdList)
-		{
-			const int32 VolumeInBytes = Array.Num() * sizeof(T);
-			FRHIBuffer* RHIBuffer = RDGBuffer->GetRHI();
-			const void* const BufferData = CmdList.LockBuffer(RHIBuffer, 0, VolumeInBytes, RLM_ReadOnly);
-			FMemory::Memcpy((void*)Array.GetData(), BufferData, VolumeInBytes);
-			CmdList.UnlockBuffer(RHIBuffer);
-		});
-}
-
-template<class T, class TRDGBuffer>
-void AddCopyToCPUPass(FRDGBuilder& GraphBuilder, TRDGBuffer RDGBuffer, T* Ptr, int NumElements)
-{
-    FCopyBufferParameters* Parameters = GraphBuilder.AllocParameters<FCopyBufferParameters>();
-    Parameters->Source = RDGBuffer;
-    GraphBuilder.AddPass(
-        RDG_EVENT_NAME("CopyToCPU"),
-        Parameters,
-        ERDGPassFlags::Readback,
-        [RDGBuffer, Ptr, NumElements](FRHICommandListImmediate& CmdList)
-        {
-            const int32 VolumeInBytes = NumElements * sizeof(T);
-            FRHIBuffer* RHIBuffer = RDGBuffer->GetRHI();
-            const void* const BufferData = CmdList.LockBuffer(RHIBuffer, 0, VolumeInBytes, RLM_ReadOnly);
-            FMemory::Memcpy((void*)Ptr, BufferData, VolumeInBytes);
-            CmdList.UnlockBuffer(RHIBuffer);
-        });
-}
-
 bool UNearestNeighborModelInstance::SetupInputs()
 {
     const bool bSuccess = Super::SetupInputs();
@@ -97,6 +55,12 @@ bool UNearestNeighborModelInstance::SetupInputs()
     return false;
 }
 
+void UNearestNeighborModelInstance::Init(USkeletalMeshComponent* SkelMeshComponent)
+{
+	UMLDeformerMorphModelInstance::Init(SkelMeshComponent);
+	InitPreviousWeights();
+}
+
 void UNearestNeighborModelInstance::Execute(float ModelWeight)
 {
 	Super::Execute(ModelWeight);
@@ -185,10 +149,10 @@ void UNearestNeighborModelInstance::UpdateWeight(TArray<float>& MorphWeights, in
 {
 	UNearestNeighborModel* NearestNeighborModel = static_cast<UNearestNeighborModel*>(Model);
 	const float DecayFactor = NearestNeighborModel->GetDecayFactor();
-	const float PreviousW = NearestNeighborModel->PreviousWeights[Index];
+	const float PreviousW = PreviousWeights[Index];
 	const float NewW = (1 - DecayFactor) * W + DecayFactor * PreviousW;
 	MorphWeights[Index] = NewW;
-	NearestNeighborModel->PreviousWeights[Index] = NewW; 
+	PreviousWeights[Index] = NewW; 
 }
 
 int32 UNearestNeighborModelInstance::FindNearestNeighbor(const FNeuralTensor& PCACoeffTensor, int32 PartId)
@@ -217,4 +181,14 @@ int32 UNearestNeighborModelInstance::FindNearestNeighbor(const FNeuralTensor& PC
 		}
 	}
 	return MinId;
-}
\ No newline at end of file
+}
+
+void UNearestNeighborModelInstance::InitPreviousWeights()
+{
+	const UNearestNeighborModel* NearestNeighborModel = Cast<UNearestNeighborModel>(Model);
+	if (NearestNeighborModel && NearestNeighborModel->GetMorphTargetSet())
+	{
+		const int32 NumWeights = NearestNeighborModel->GetMorphTargetSet()->MorphBuffers.GetNumMorphs();
+		PreviousWeights.SetNumZeroed(NumWeights);
+	}
+}
diff --git a/Engine/Plugins/Animation/MLDeformer/NearestNeighborModel/Source/NearestNeighborModel/Public/NearestNeighborModel.h b/Engine/Plugins/Animation/MLDeformer/NearestNeighborModel/Source/NearestNeighborModel/Public/NearestNeighborModel.h
--- a/Engine/Plugins/Animation/MLDeformer/NearestNeighborModel/Source/NearestNeighborModel/Public/NearestNeighborModel.h
+++ b/Engine/Plugins/Animation/MLDeformer/NearestNeighborModel/Source/NearestNeighborModel/Public/NearestNeighborModel.h
@@ -145,8 +145,8 @@ public:
 	virtual void PostLoad() override;
 
 	// UMLDeformerModel overrides.
-	virtual UMLDeformerModelInstance* CreateModelInstance(UMLDeformerComponent* Component);
-	virtual UMLDeformerInputInfo* CreateInputInfo();
+	virtual UMLDeformerModelInstance* CreateModelInstance(UMLDeformerComponent* Component) override;
+	virtual UMLDeformerInputInfo* CreateInputInfo() override;
 	virtual FString GetDisplayName() const override { return "Nearest Neighbor Model"; }
 	// ~END UMLDeformerModel overrides.
 
@@ -213,7 +213,6 @@ public:
 	bool CheckPCAData(int32 PartId) const;
 
 	void InitInputInfo();
-	void InitPreviousWeights();
 
 	void ResetMorphBuffers();
 
@@ -367,6 +366,4 @@ public:
 
 	UPROPERTY(EditAnywhere, BlueprintReadWrite, Category = "Nearest Neighbors", META = (ClampMin = "0", ClampMax = "1"))
 	float NearestNeighborOffsetWeight = 1.0f;
-
-	TArray<float> PreviousWeights; 
 };
diff --git a/Engine/Plugins/Animation/MLDeformer/NearestNeighborModel/Source/NearestNeighborModel/Public/NearestNeighborModelInstance.h b/Engine/Plugins/Animation/MLDeformer/NearestNeighborModel/Source/NearestNeighborModel/Public/NearestNeighborModelInstance.h
--- a/Engine/Plugins/Animation/MLDeformer/NearestNeighborModel/Source/NearestNeighborModel/Public/NearestNeighborModelInstance.h
+++ b/Engine/Plugins/Animation/MLDeformer/NearestNeighborModel/Source/NearestNeighborModel/Public/NearestNeighborModelInstance.h
@@ -13,6 +13,7 @@ class NEARESTNEIGHBORMODEL_API UNearestNeighborModelInstance
 
 public:
 	// UMLDeformerModelInstance overrides
+    virtual void Init(USkeletalMeshComponent* SkelMeshComponent) override;
     virtual void Execute(float ModelWeight) override;
     virtual bool SetupInputs() override;
 
@@ -21,6 +22,8 @@ public:
     int32 NeighborIdNum() const { return NearestNeighborIds.Num(); }
 #endif
 
+    void InitPreviousWeights();
+
 protected:
     virtual int64 SetBoneTransforms(float* OutputBuffer, int64 OutputBufferSize, int64 StartIndex) override;
     // ~END UMLDeformerModelInstance overrides
@@ -33,4 +36,6 @@ private:
 #if WITH_EDITORONLY_DATA
     TArray<uint32> NearestNeighborIds;
 #endif
+
+    TArray<float> PreviousWeights;
 };
diff --git a/Engine/Plugins/Animation/MLDeformer/NearestNeighborModel/Source/NearestNeighborModelEditor/Private/NearestNeighborEditorModel.cpp b/Engine/Plugins/Animation/MLDeformer/NearestNeighborModel/Source/NearestNeighborModelEditor/Private/NearestNeighborEditorModel.cpp
--- a/Engine/Plugins/Animation/MLDeformer/NearestNeighborModel/Source/NearestNeighborModelEditor/Private/NearestNeighborEditorModel.cpp
+++ b/Engine/Plugins/Animation/MLDeformer/NearestNeighborModel/Source/NearestNeighborModelEditor/Private/NearestNeighborEditorModel.cpp
@@ -75,7 +75,7 @@ namespace UE::NearestNeighborModel
 		FMLDeformerMorphModelEditorModel::OnPostTraining(TrainingResult, bUsePartiallyTrainedWhenAborted);
 		if (TrainingResult == ETrainingResult::Success || (TrainingResult == ETrainingResult::Aborted && bUsePartiallyTrainedWhenAborted))
 		{
-			GetNearestNeighborModel()->InitPreviousWeights();
+			InitTestMLDeformerPreviousWeights();
 		}
 	}
 
@@ -135,17 +135,6 @@ namespace UE::NearestNeighborModel
 			NearestNeighborActors.Reset();
 		}
 		NearestNeighborActors.SetNumZeroed(NumParts);
-
-		UMLDeformerComponent* MLDeformerComponent = nullptr;
-		FMLDeformerEditorActor* TestActor = FindEditorActor(ActorID_Test_MLDeformed);
-		if (TestActor)
-		{
-			AActor* Actor = TestActor->GetActor();
-			if (Actor)
-			{
-				MLDeformerComponent = Actor->FindComponentByClass<UMLDeformerComponent>();
-			}
-		}
 		
 		for(int32 PartId = StartIndex; PartId < NumParts; PartId++)
 		{
@@ -162,7 +151,7 @@ namespace UE::NearestNeighborModel
 				LOCTEXT("TestNearestNeighborLabelText", "Nearest Neigbors"),
 				false);
 			FNearestNeighborEditorModelActor* NearestNeighborActor = static_cast<FNearestNeighborEditorModelActor*>(EditorActors.Last());
-
+			UMLDeformerComponent* MLDeformerComponent = GetTestMLDeformerComponent();
 			NearestNeighborActor->InitNearestNeighborActor(PartId, MLDeformerComponent);
 			NearestNeighborActor->SetMeshOffsetFactor(2.0f);
 			NearestNeighborActors[PartId] = NearestNeighborActor;
@@ -516,6 +505,34 @@ namespace UE::NearestNeighborModel
 			SkelMeshComponent->RefreshExternalMorphTargetWeights();
 		}
 	}
+
+	UMLDeformerComponent* FNearestNeighborEditorModel::GetTestMLDeformerComponent()
+	{
+		UMLDeformerComponent* MLDeformerComponent = nullptr;
+		FMLDeformerEditorActor* TestActor = FindEditorActor(ActorID_Test_MLDeformed);
+		if (TestActor)
+		{
+			AActor* Actor = TestActor->GetActor();
+			if (Actor)
+			{
+				MLDeformerComponent = Actor->FindComponentByClass<UMLDeformerComponent>();
+			}
+		}
+		return MLDeformerComponent;
+	}
+
+	void FNearestNeighborEditorModel::InitTestMLDeformerPreviousWeights()
+	{
+		UMLDeformerComponent* MLDeformerComponent = GetTestMLDeformerComponent();
+		if (MLDeformerComponent)
+		{
+			UNearestNeighborModelInstance* ModelInstance = static_cast<UNearestNeighborModelInstance*>(MLDeformerComponent->GetModelInstance());
+			if (ModelInstance)
+			{
+				ModelInstance->InitPreviousWeights();
+			}
+		}
+	}
 }	// namespace UE::NearestNeighborModel
 
 #undef LOCTEXT_NAMESPACE
diff --git a/Engine/Plugins/Animation/MLDeformer/NearestNeighborModel/Source/NearestNeighborModelEditor/Private/NearestNeighborModelDetails.cpp b/Engine/Plugins/Animation/MLDeformer/NearestNeighborModel/Source/NearestNeighborModelEditor/Private/NearestNeighborModelDetails.cpp
--- a/Engine/Plugins/Animation/MLDeformer/NearestNeighborModel/Source/NearestNeighborModelEditor/Private/NearestNeighborModelDetails.cpp
+++ b/Engine/Plugins/Animation/MLDeformer/NearestNeighborModel/Source/NearestNeighborModelEditor/Private/NearestNeighborModelDetails.cpp
@@ -172,7 +172,7 @@ namespace UE::NearestNeighborModel
 				.OnClicked_Lambda([this]
 				{
 					NearestNeighborEditorModel->OnMorphTargetUpdate();
-					NearestNeighborModel->InitPreviousWeights();
+					NearestNeighborEditorModel->InitTestMLDeformerPreviousWeights();
 					return FReply::Handled();
 				})
 			];
diff --git a/Engine/Plugins/Animation/MLDeformer/NearestNeighborModel/Source/NearestNeighborModelEditor/Public/NearestNeighborEditorModel.h b/Engine/Plugins/Animation/MLDeformer/NearestNeighborModel/Source/NearestNeighborModelEditor/Public/NearestNeighborEditorModel.h
--- a/Engine/Plugins/Animation/MLDeformer/NearestNeighborModel/Source/NearestNeighborModelEditor/Public/NearestNeighborEditorModel.h
+++ b/Engine/Plugins/Animation/MLDeformer/NearestNeighborModel/Source/NearestNeighborModelEditor/Public/NearestNeighborEditorModel.h
@@ -9,6 +9,7 @@
 class UMLDeformerAsset;
 class USkeletalMesh;
 class UGeometryCache;
+class UMLDeformerComponent;
 class UNearestNeighborModel;
 class UNearestNeighborModelVizSettings;
 
@@ -68,6 +69,8 @@ namespace UE::NearestNeighborModel
 
 		void OnMorphTargetUpdate();
 		uint8 GetMorphTargetUpdateResult() { return MorphTargetUpdateResult; }
+		UMLDeformerComponent* GetTestMLDeformerComponent();
+		void InitTestMLDeformerPreviousWeights();
 
 
 	protected:
