NearestNeighborEditor: Consolidate update buttons
Add more descriptive error messages and remove crashing.

diff --git a/Engine/Plugins/Animation/MLDeformer/NearestNeighborModel/Source/NearestNeighborModel/Private/NearestNeighborModel.cpp b/Engine/Plugins/Animation/MLDeformer/NearestNeighborModel/Source/NearestNeighborModel/Private/NearestNeighborModel.cpp
--- a/Engine/Plugins/Animation/MLDeformer/NearestNeighborModel/Source/NearestNeighborModel/Private/NearestNeighborModel.cpp
+++ b/Engine/Plugins/Animation/MLDeformer/NearestNeighborModel/Source/NearestNeighborModel/Private/NearestNeighborModel.cpp
@@ -96,6 +96,16 @@ void UNearestNeighborModel::ClipInputs(float* InputPtr, int NumInputs)
 	}	
 }
 
+int32 UNearestNeighborModel::GetTotalNumPCACoeffs() const
+{
+	int32 TotalNumPCACoeffs = 0;
+	for (int32 PartId = 0; PartId < GetNumParts(); PartId++)
+	{
+		TotalNumPCACoeffs += GetPCACoeffNum(PartId);
+	}
+	return TotalNumPCACoeffs;
+}
+
 int32 UNearestNeighborModel::GetTotalNumNeighbors() const
 {
 	int32 TotalNumNeighbors = 0;
@@ -129,7 +139,7 @@ void UNearestNeighborModel::UpdatePCACoeffNums()
 	for(int32 PartId = 0; PartId < GetNumParts(); PartId++)
 	{
 		ClothPartData[PartId].PCACoeffStart = PCACoeffStart;
-		PCACoeffStart+=ClothPartData[PartId].PCACoeffNum;
+		PCACoeffStart += ClothPartData[PartId].PCACoeffNum;
 	}
 }
 
@@ -147,11 +157,13 @@ void UNearestNeighborModel::UpdateNetworkOutputDim()
 	}
 }
 
-void UNearestNeighborModel::UpdateClothPartData()
+UE::NearestNeighborModel::EUpdateResult UNearestNeighborModel::UpdateClothPartData()
 {
-	if(ClothPartData.Num() == 0)
+	using namespace UE::NearestNeighborModel;
+	if(ClothPartEditorData.Num() == 0)
 	{
 		UE_LOG(LogNearestNeighborModel, Error, TEXT("There should be at least 1 cloth part"));
+		return EUpdateResult::ERROR;
 	}
 
 	ClothPartData.SetNum(ClothPartEditorData.Num());
@@ -169,7 +181,6 @@ void UNearestNeighborModel::UpdateClothPartData()
 				ClothPartData[PartId].VertexMean.SetNumZeroed(ClothPartData[PartId].NumVertices * 3);
 				ClothPartData[PartId].PCABasis.SetNumZeroed(ClothPartData[PartId].NumVertices * 3 * ClothPartData[PartId].PCACoeffNum);
 
-				// NumNeighbors cannot be 0 because of the NearestNeighborModel.usf shader.
 				// Init default neighbor data.
 				ClothPartData[PartId].NeighborCoeffs.SetNumZeroed(ClothPartData[PartId].PCACoeffNum);
 				ClothPartData[PartId].NeighborOffsets.SetNumZeroed(ClothPartData[PartId].NumVertices * 3);
@@ -178,8 +189,8 @@ void UNearestNeighborModel::UpdateClothPartData()
 		}
 		else
 		{
-			UE_LOG(LogNearestNeighborModel, Error, TEXT("%s does not exist"), *ClothPartEditorData[PartId].VertexMapPath);
-			return;
+			UE_LOG(LogNearestNeighborModel, Error, TEXT("Part %d txt path %s does not exist"), PartId, *ClothPartEditorData[PartId].VertexMapPath);
+			return EUpdateResult::ERROR;
 		}
 	}
 	UpdatePCACoeffNums();
@@ -188,6 +199,7 @@ void UNearestNeighborModel::UpdateClothPartData()
 	UpdateNetworkInputDim();
 	UpdateNetworkOutputDim();
 	bClothPartDataValid = true;
+	return EUpdateResult::SUCCESS;
 }
 
 TObjectPtr<UAnimSequence> UNearestNeighborModel::GetNearestNeighborSkeletons(int32 PartId)
@@ -228,7 +240,7 @@ const TObjectPtr<UGeometryCache> UNearestNeighborModel::GetNearestNeighborCache(
 
 int32 UNearestNeighborModel::GetNumNeighborsFromGeometryCache(int32 PartId) const
 {
-	const TObjectPtr<UGeometryCache> Cache = GetNearestNeighborCache(PartId);
+	const UGeometryCache* Cache = GetNearestNeighborCache(PartId);
 	if (Cache)
 	{
 		const int32 StartFrame = Cache->GetStartFrame();
@@ -241,6 +253,19 @@ int32 UNearestNeighborModel::GetNumNeighborsFromGeometryCache(int32 PartId) cons
 	}
 }
 
+int32 UNearestNeighborModel::GetNumNeighborsFromAnimSequence(int32 PartId) const
+{
+	const UAnimSequence* Anim = GetNearestNeighborSkeletons(PartId);
+	if (Anim)
+	{
+		return Anim->GetDataModel()->GetNumberOfKeys();
+	}
+	else
+	{
+		return 0;
+	}
+}
+
 void UNearestNeighborModel::UpdateNetworkSize()
 {
 	UNeuralNetwork* Network = GetNeuralNetwork();
diff --git a/Engine/Plugins/Animation/MLDeformer/NearestNeighborModel/Source/NearestNeighborModel/Public/NearestNeighborModel.h b/Engine/Plugins/Animation/MLDeformer/NearestNeighborModel/Source/NearestNeighborModel/Public/NearestNeighborModel.h
--- a/Engine/Plugins/Animation/MLDeformer/NearestNeighborModel/Source/NearestNeighborModel/Public/NearestNeighborModel.h
+++ b/Engine/Plugins/Animation/MLDeformer/NearestNeighborModel/Source/NearestNeighborModel/Public/NearestNeighborModel.h
@@ -115,11 +115,14 @@ struct FClothPartData
 	TArray<float> NeighborOffsets;
 };
 
-enum EResultMessage : uint8
+namespace UE::NearestNeighborModel
 {
-	SUCCESS = 0,
-	ERROR = 1,
-	WARNING = 2
+	enum EUpdateResult : uint8
+	{
+		SUCCESS = 0,
+		ERROR = 1,
+		WARNING = 2
+	};
 };
 
 /**
@@ -165,6 +168,7 @@ public:
 	UFUNCTION(BlueprintPure, Category = "Nearest Neighbor Model")
 	int32 GetNumNeighbors(int32 PartId) const { return ClothPartData[PartId].NumNeighbors; }
 
+	int32 GetTotalNumPCACoeffs() const;
 	int32 GetTotalNumNeighbors() const;
 
 	float GetDecayFactor() const { return DecayFactor; };
@@ -222,20 +226,24 @@ public:
 	TObjectPtr<UGeometryCache> GetNearestNeighborCache(int32 PartId);
 	const TObjectPtr<UGeometryCache> GetNearestNeighborCache(int32 PartId) const;
 	int32 GetNumNeighborsFromGeometryCache(int32 PartId) const;
+	int32 GetNumNeighborsFromAnimSequence(int32 PartId) const;
 
 	void UpdateNetworkInputDim();
 	void UpdateNetworkOutputDim();
-	void UpdateClothPartData();
+	UE::NearestNeighborModel::EUpdateResult UpdateClothPartData();
 	void UpdatePCACoeffNums();
 	void UpdateNetworkSize();
 	void UpdateMorphTargetSize();
 
-	void InvalidateClothPartData() { bClothPartDataValid = false; bNearestNeighborDataValid = false; }
+	void InvalidateClothPartData() { bClothPartDataValid = false; bNearestNeighborDataValid = false; bMorphTargetDataValid = false; }
 	void ValidateClothPartData() { bClothPartDataValid = true; }
-	void InvalidateNearestNeighborData() { bNearestNeighborDataValid = false; }
+	void InvalidateNearestNeighborData() { bNearestNeighborDataValid = false; bMorphTargetDataValid = false; }
 	void ValidateNearestNeighborData() {bNearestNeighborDataValid = true; }
-	bool IsClothPartDataValid() { return bClothPartDataValid; }
-	bool IsNearestNeighborDataValid() { return bNearestNeighborDataValid; }
+	void InvalidateMorphTargetData() { bMorphTargetDataValid = false; }
+	void ValidateMorphTargetData() { bMorphTargetDataValid = true;}
+	bool IsClothPartDataValid() const { return bClothPartDataValid; }
+	bool IsNearestNeighborDataValid() const { return bNearestNeighborDataValid; }
+	bool IsMorphTargetDataValid() const { return bMorphTargetDataValid; }
 #endif
 
 #if WITH_EDITOR
@@ -320,6 +328,9 @@ protected:
 	UPROPERTY()
 	bool bNearestNeighborDataValid = false;
 
+	UPROPERTY()
+	bool bMorphTargetDataValid = false;
+
 	UPROPERTY(EditAnywhere, BlueprintReadWrite, Category = "Nearest Neighbors")
 	bool bUsePartOnlyMesh = true;
 
diff --git a/Engine/Plugins/Animation/MLDeformer/NearestNeighborModel/Source/NearestNeighborModelEditor/Private/NearestNeighborEditorModel.cpp b/Engine/Plugins/Animation/MLDeformer/NearestNeighborModel/Source/NearestNeighborModelEditor/Private/NearestNeighborEditorModel.cpp
--- a/Engine/Plugins/Animation/MLDeformer/NearestNeighborModel/Source/NearestNeighborModelEditor/Private/NearestNeighborEditorModel.cpp
+++ b/Engine/Plugins/Animation/MLDeformer/NearestNeighborModel/Source/NearestNeighborModelEditor/Private/NearestNeighborEditorModel.cpp
@@ -63,11 +63,17 @@ namespace UE::NearestNeighborModel
 			GetEditor()->GetModelDetailsView()->ForceRefresh();
 		}
 
-		if (Property->GetFName() == UNearestNeighborModel::GetNearestNeighborDataPropertyName())
+		if (Property->GetFName() == UNearestNeighborModel::GetNearestNeighborDataPropertyName() || (PropertyChangedEvent.MemberProperty != nullptr && PropertyChangedEvent.MemberProperty->GetFName() == UNearestNeighborModel::GetNearestNeighborDataPropertyName()) || Property->GetFName() == UNearestNeighborModel::GetUsePartOnlyMeshPropertyName())
 		{
 			GetNearestNeighborModel()->InvalidateNearestNeighborData();
 			GetEditor()->GetModelDetailsView()->ForceRefresh();
 		}
+
+		if (Property->GetFName() == UNearestNeighborModel::GetMorphTargetDeltaThresholdPropertyName() || Property->GetFName() == UNearestNeighborModel::GetMorphTargetErrorTolerancePropertyName())
+		{
+			GetNearestNeighborModel()->InvalidateMorphTargetData();
+			GetEditor()->GetModelDetailsView()->ForceRefresh();
+		}
 	}
 
 	void FNearestNeighborEditorModel::OnPostTraining(ETrainingResult TrainingResult, bool bUsePartiallyTrainedWhenAborted)
@@ -170,12 +176,29 @@ namespace UE::NearestNeighborModel
 		NearestNeighborModel->InitInputInfo();
 	}
 
+	bool HasError(uint8 Flag)
+	{
+		return (Flag & EUpdateResult::ERROR) != 0;
+	}
+
 	ETrainingResult FNearestNeighborEditorModel::Train()
 	{
+		UNearestNeighborModel* NearestNeighborModel = GetNearestNeighborModel();
+		MorphTargetUpdateResult = EUpdateResult::SUCCESS;
+		if (!NearestNeighborModel->IsClothPartDataValid())
+		{
+			MorphTargetUpdateResult |= NearestNeighborModel->UpdateClothPartData();
+			if (HasError(MorphTargetUpdateResult))
+			{
+				UpdateNearestNeighborActors();
+				return ETrainingResult::FailOnData;
+			}
+		}
+
 		return TrainModel<UNearestNeighborTrainingModel>(this);
 	}
 
-	bool FNearestNeighborEditorModel::SetSamplerPartData(const int32 PartId)
+	uint8 FNearestNeighborEditorModel::SetSamplerPartData(const int32 PartId)
 	{
 		FNearestNeighborGeomCacheSampler* GeomCacheSampler = static_cast<FNearestNeighborGeomCacheSampler*>(GetGeomCacheSampler());
 		UNearestNeighborModel* NearestNeighborModel = GetNearestNeighborModel();
@@ -189,8 +212,10 @@ namespace UE::NearestNeighborModel
 
 			if (SkeletalMeshComponent && GeometryCacheComponent && AnimSequence && GeometryCache)
 			{
-				const int32 NumFrames = GeometryCache->GetEndFrame() - GeometryCache->GetStartFrame() + 1;
-				if (AnimSequence->GetDataModel()->GetNumberOfKeys() == NumFrames)
+				const int32 NumNeighborsFromGeomCache = NearestNeighborModel->GetNumNeighborsFromGeometryCache(PartId);
+				const int32 NumNeighborsFromAnimSequence = NearestNeighborModel->GetNumNeighborsFromAnimSequence(PartId);
+				const int32 NumFrames = FMath::Min(NumNeighborsFromGeomCache, NumNeighborsFromAnimSequence);
+				if (NumFrames > 0)
 				{
 					AnimSequence->bUseRawDataOnly = true;
 					AnimSequence->Interpolation = EAnimInterpolationType::Step;
@@ -209,23 +234,48 @@ namespace UE::NearestNeighborModel
 					GeometryCacheComponent->SetManualTick(true);
 					GeometryCacheComponent->SetPlaybackSpeed(1.0f);
 					GeometryCacheComponent->Play();
-					GeomCacheSampler->GeneratePartMeshMappings(GetNearestNeighborModel()->PartVertexMap(PartId), NearestNeighborModel->GetUsePartOnlyMesh());
+					uint8 ReturnCode = GeomCacheSampler->GeneratePartMeshMappings(GetNearestNeighborModel()->PartVertexMap(PartId), NearestNeighborModel->GetUsePartOnlyMesh());
+					if (HasError(ReturnCode))
+					{
+						return ReturnCode;
+					}
+
+					if (GeomCacheSampler->IsMeshMappingsEmpty())
+					{
+						UE_LOG(LogNearestNeighborModel, Error, TEXT("Failed to generate mappings between skeletal mesh and geometry cache."));
+						return EUpdateResult::ERROR;
+					}
 
 					NumTrainingFramesOverride = NumFrames;
 
-					return true;
+					if (NumNeighborsFromGeomCache == NumNeighborsFromAnimSequence)
+					{
+						return EUpdateResult::SUCCESS;
+					}
+					else
+					{
+						UE_LOG(LogNearestNeighborModel, Warning, TEXT("NearestNeighborData: part %d frame mismatch: AnimSequence has %d frames and GeometryCache has %d frames. Using %f frames only."), PartId, NumNeighborsFromAnimSequence, NumNeighborsFromGeomCache, NumFrames);
+						return EUpdateResult::WARNING;
+					}
 				}
 				else
 				{
-					UE_LOG(LogNearestNeighborModel, Error, TEXT("Part %d frame mismatch: AnimSequence has %d frames and GeometryCache has %d frames"), PartId, AnimSequence->GetDataModel()->GetNumberOfKeys(), NumFrames);
+					UE_LOG(LogNearestNeighborModel, Warning, TEXT("Part %d: AnimSequence or GeometryCache has zero frames"), PartId);
+					return EUpdateResult::WARNING;
 				}
 			}
+			else if (AnimSequence == nullptr && GeometryCache == nullptr)
+			{
+				return EUpdateResult::SUCCESS;
+			}
 			else
 			{
-				UE_LOG(LogNearestNeighborModel, Error, TEXT("Part %d: AnimSequence or GeometryCache is null"), PartId);
+				UE_LOG(LogNearestNeighborModel, Warning, TEXT("Part %d is skipped because AnimSequence or GeometryCache is None"), PartId);
+				return EUpdateResult::WARNING;
 			}
 		}
-		return false;
+		UE_LOG(LogNearestNeighborModel, Error, TEXT("SetSamplerPartData: unknown error"));
+		return EUpdateResult::ERROR;
 	}
 
 	int32 FNearestNeighborEditorModel::GetNumParts()
@@ -285,34 +335,36 @@ namespace UE::NearestNeighborModel
 			{
 				CreateNearestNeighborActors(EditorWorld, NearestNeighborActors.Num());
 			}
-		}
 
-		for (int32 PartId = 0; PartId < NearestNeighborActors.Num(); PartId++)
-		{
-			const TObjectPtr<UGeometryCache> GeometryCache = GetNearestNeighborModel()->GetNearestNeighborCache(PartId);
-			NearestNeighborActors[PartId]->GetGeometryCacheComponent()->SetGeometryCache(GeometryCache);
+			for (int32 PartId = 0; PartId < NearestNeighborActors.Num(); PartId++)
+			{
+				const TObjectPtr<UGeometryCache> GeometryCache = GetNearestNeighborModel()->GetNearestNeighborCache(PartId);
+				NearestNeighborActors[PartId]->GetGeometryCacheComponent()->SetGeometryCache(GeometryCache);
+			}
 		}
 	}
 
-	void FNearestNeighborEditorModel::UpdateNearestNeighborData()
+	uint8 FNearestNeighborEditorModel::UpdateNearestNeighborData()
 	{
 		UNearestNeighborModel *NearestNeighborModel = static_cast<UNearestNeighborModel*>(Model);
 		check(NearestNeighborModel != nullptr);
 
+		if (NearestNeighborModel->GetNumParts() == 0)
+		{
+			return EUpdateResult::SUCCESS;
+		}
+
 		UNeuralNetwork *NeuralNetwork = NearestNeighborModel->GetNeuralNetwork();
 		if(NeuralNetwork && NeuralNetwork->IsLoaded())
 		{
 			const FString SavePath = GetTrainedNetworkOnnxFile();
-			if (!FPaths::FileExists(SavePath) || !NearestNeighborModel->GetUseFileCache())
-			{
-				UE_LOG(LogNearestNeighborModel, Display, TEXT("Saving to %s"), *SavePath);
-				NeuralNetwork->Save(SavePath);
-			}
+			UE_LOG(LogNearestNeighborModel, Display, TEXT("Saving to %s"), *SavePath);
+			NeuralNetwork->Save(SavePath);
 		}
 		else
 		{
-			UE_LOG(LogNearestNeighborModel, Warning, TEXT("Network is not available. Nothing will be done."));
-			return;
+			UE_LOG(LogNearestNeighborModel, Display, TEXT("Network not loaded. A network needs to be trained."));
+			return EUpdateResult::WARNING;
 		}
 
 		FNearestNeighborGeomCacheSampler* GeomCacheSampler = static_cast<FNearestNeighborGeomCacheSampler*>(GetGeomCacheSampler());
@@ -325,12 +377,17 @@ namespace UE::NearestNeighborModel
 			{
 				UNearestNeighborTrainingModel *TrainingModel = InitTrainingModel<UNearestNeighborTrainingModel>(this);
 				check(TrainingModel != nullptr);
-				TrainingModel->UpdateNearestNeighborData();
+				uint8 ReturnCode = TrainingModel->UpdateNearestNeighborData();
 				ResetSamplerData();
-				UpdateNearestNeighborActors();
-				NearestNeighborModel->ValidateNearestNeighborData();
+				if (HasError(ReturnCode) == 0)
+				{
+					NearestNeighborModel->ValidateNearestNeighborData();
+				}
+				return ReturnCode;
 			}
 		}
+		UE_LOG(LogNearestNeighborModel, Warning, TEXT("GeomCacheSampler is empty. Nearest neighbor data is not updated."))
+		return EUpdateResult::ERROR;
 	}
 
 	void FNearestNeighborEditorModel::KMeansClusterPoses()
@@ -386,28 +443,51 @@ namespace UE::NearestNeighborModel
 			}
 		}
 	}
-#define RETURN_ON_ERROR(Flag, Func)\
-	Flag |= Func;\
-	if ((Flag & ERROR) != 0)\
-	{\
-		GetEditor()->GetModelDetailsView()->ForceRefresh();\
-		return;\
-	}
 	
 	void FNearestNeighborEditorModel::OnMorphTargetUpdate()
 	{
-		MorphTargetUpdateResult = SUCCESS;
-		RETURN_ON_ERROR(MorphTargetUpdateResult, InitMorphTargets());
-		RefreshMorphTargets();
-		GetNearestNeighborModel()->UpdateNetworkSize();
-		GetNearestNeighborModel()->UpdateMorphTargetSize();
+		MorphTargetUpdateResult = EUpdateResult::SUCCESS;
+
+		UNearestNeighborModel* NearestNeighborModel = GetNearestNeighborModel();
+		if (!NearestNeighborModel->IsClothPartDataValid())
+		{
+			MorphTargetUpdateResult |= NearestNeighborModel->UpdateClothPartData();
+			if (HasError(MorphTargetUpdateResult))
+			{
+				UpdateNearestNeighborActors();
+				return;
+			}
+		}
+		MorphTargetUpdateResult |= WarnIfNetworkInvalid();
+
+		if (!NearestNeighborModel->IsNearestNeighborDataValid())
+		{
+			MorphTargetUpdateResult |= UpdateNearestNeighborData();
+			UpdateNearestNeighborActors();
+			if (HasError(MorphTargetUpdateResult))
+			{
+				return;
+			}
+		}
 
-		GetEditor()->GetModelDetailsView()->ForceRefresh();
+		if (IsNeuralNetworkLoaded())
+		{
+			MorphTargetUpdateResult |= InitMorphTargets();
+			if (HasError(MorphTargetUpdateResult))
+			{
+				return;
+			}
+			RefreshMorphTargets();
+			GetNearestNeighborModel()->UpdateNetworkSize();
+			GetNearestNeighborModel()->UpdateMorphTargetSize();
+		}
+		InitTestMLDeformerPreviousWeights();
+		GetNearestNeighborModel()->ValidateMorphTargetData();
 	}
 
 	uint8 FNearestNeighborEditorModel::InitMorphTargets()
 	{
-		EResultMessage Result = SUCCESS;
+		uint8 Result = EUpdateResult::SUCCESS;
 		UNearestNeighborModel* NearestNeighborModel = GetNearestNeighborModel();
 
 		const USkeletalMesh* SkelMesh = Model->GetSkeletalMesh();
@@ -415,27 +495,27 @@ namespace UE::NearestNeighborModel
 		if (!SkelMesh || NumBaseMeshVerts == 0)
 		{
 			UE_LOG(LogNearestNeighborModel, Error, TEXT("SkeletalMesh is empty. No morph targets are generated"));
-			return ERROR;
+			return EUpdateResult::ERROR;
 		}
 
 		if (Model->GetVertexMap().IsEmpty())
 		{
 			UE_LOG(LogNearestNeighborModel, Error, TEXT("VertexMap of the skeletal mesh is empty. No morph targets are generated"));
-			return ERROR;
+			return EUpdateResult::ERROR;
 		}
 
 		const int32 NumParts = NearestNeighborModel->GetNumParts();
 		if (NumParts == 0)
 		{
 			UE_LOG(LogNearestNeighborModel, Error, TEXT("There are no cloth parts. No morph targets are generated"));
-			return ERROR;
+			return EUpdateResult::ERROR;
 		}
 
 		const int32 NumImportedModelVerts = FMath::Max(NearestNeighborModel->GetVertexMap()) + 1;
 		if (NumImportedModelVerts != NumBaseMeshVerts)
 		{
 			UE_LOG(LogNearestNeighborModel, Error, TEXT("Vertex count mismatch: imported model of SkeletalMesh has %d vertices and cached SkeletalMesh has %d vertices"), NumImportedModelVerts, NumBaseMeshVerts);
-			return ERROR;
+			return EUpdateResult::ERROR;
 		}
 
 		TArray<FVector3f> Deltas;
@@ -454,7 +534,7 @@ namespace UE::NearestNeighborModel
 			if (VertexMap.IsEmpty())
 			{
 				UE_LOG(LogNearestNeighborModel, Warning, TEXT("Cloth part %d has empty vertex map. No morph targets are generated for this part."), PartId);
-				Result = WARNING;
+				Result = EUpdateResult::WARNING;
 			}
 			AddFloatArrayToDeltaArray(NearestNeighborModel->ClothPartData[PartId].VertexMean, VertexMap, Deltas, 0);
 			AddFloatArrayToDeltaArray(NearestNeighborModel->ClothPartData[PartId].PCABasis, VertexMap, Deltas);
@@ -469,7 +549,7 @@ namespace UE::NearestNeighborModel
 		if (Deltas.Num() == 0)
 		{
 			UE_LOG(LogNearestNeighborModel, Error, TEXT("All cloth parts are empty. No morph targets are generated."));
-			return ERROR;
+			return EUpdateResult::ERROR;
 		}
 
 		const int32 LOD = 0;
@@ -482,7 +562,7 @@ namespace UE::NearestNeighborModel
 		if (MorphBuffers.GetNumBatches() <= 0)
 		{
 			UE_LOG(LogNearestNeighborModel, Warning, TEXT("Morph buffer is empty. It is possible that all deltas are zero. No morph targets are generated."));
-			Result = WARNING;
+			Result = EUpdateResult::WARNING;
 			NearestNeighborModel->ResetMorphBuffers();
 		}
 
@@ -533,6 +613,35 @@ namespace UE::NearestNeighborModel
 			}
 		}
 	}
+
+	uint8 FNearestNeighborEditorModel::WarnIfNetworkInvalid()
+	{
+		const UNearestNeighborModel *NearestNeighborModel = static_cast<UNearestNeighborModel*>(Model);
+
+		const UNeuralNetwork *NeuralNetwork = NearestNeighborModel->GetNeuralNetwork();
+		const UMLDeformerComponent* MLDeformerComponent = GetTestMLDeformerComponent();
+		if(NeuralNetwork && NeuralNetwork->IsLoaded() && MLDeformerComponent && MLDeformerComponent->GetModelInstance())
+		{
+			const int32 NeuralNetworkInferenceHandle = MLDeformerComponent->GetModelInstance()->GetNeuralNetworkInferenceHandle();
+			const FNeuralTensor& OutputTensor = NeuralNetwork->GetOutputTensorForContext(NeuralNetworkInferenceHandle);
+			const int32 NumNetworkWeights = OutputTensor.Num();
+
+			const int32 NumPCACoeffs = NearestNeighborModel->GetTotalNumPCACoeffs();
+			if (NumNetworkWeights != NumPCACoeffs)
+			{
+				UE_LOG(LogNearestNeighborModel, Warning, TEXT("Network output dimension %d is not equal to number of morph targets %d. Network needs to be re-trained and no deformation will be applied."), NumNetworkWeights, NumPCACoeffs);
+				return EUpdateResult::WARNING;
+			}
+		}
+		return EUpdateResult::SUCCESS;
+	}
+
+	bool FNearestNeighborEditorModel::IsNeuralNetworkLoaded()
+	{
+		const UNearestNeighborModel *NearestNeighborModel = static_cast<UNearestNeighborModel*>(Model);
+		const UNeuralNetwork *NeuralNetwork = NearestNeighborModel->GetNeuralNetwork();
+		return NeuralNetwork && NeuralNetwork->IsLoaded();
+	}
 }	// namespace UE::NearestNeighborModel
 
 #undef LOCTEXT_NAMESPACE
diff --git a/Engine/Plugins/Animation/MLDeformer/NearestNeighborModel/Source/NearestNeighborModelEditor/Private/NearestNeighborGeomCacheSampler.cpp b/Engine/Plugins/Animation/MLDeformer/NearestNeighborModel/Source/NearestNeighborModelEditor/Private/NearestNeighborGeomCacheSampler.cpp
--- a/Engine/Plugins/Animation/MLDeformer/NearestNeighborModel/Source/NearestNeighborModelEditor/Private/NearestNeighborGeomCacheSampler.cpp
+++ b/Engine/Plugins/Animation/MLDeformer/NearestNeighborModel/Source/NearestNeighborModelEditor/Private/NearestNeighborGeomCacheSampler.cpp
@@ -17,7 +17,7 @@
 using namespace UE::MLDeformer;
 namespace UE::NearestNeighborModel
 {
-	void FNearestNeighborGeomCacheSampler::SamplePart(int32 InAnimFrameIndex, const TArray<uint32>& VertexMap)
+	uint8 FNearestNeighborGeomCacheSampler::SamplePart(int32 InAnimFrameIndex, const TArray<uint32>& VertexMap)
 	{
 		FMLDeformerSampler::Sample(InAnimFrameIndex);
 		USkeletalMesh* SkeletalMesh = SkeletalMeshComponent.Get() ? SkeletalMeshComponent->GetSkeletalMeshAsset() : nullptr;
@@ -30,7 +30,11 @@ namespace UE::NearestNeighborModel
 			const FSkeletalMeshLODModel& LODModel = ImportedModel->LODModels[LODIndex];
 			const TArray<FSkelMeshImportedMeshInfo>& SkelMeshInfos = LODModel.ImportedMeshInfos;
 
-			check(MeshMappings.Num() >= 1);
+			if (MeshMappings.IsEmpty())
+			{
+				UE_LOG(LogNearestNeighborModel, Error, TEXT("SamplePart: MeshMappings is empty."));
+				return EUpdateResult::ERROR;
+			}
 			const UE::MLDeformer::FMLDeformerGeomCacheMeshMapping& MeshMapping = MeshMappings[0]; 
 
 			check(SkelMeshInfos.Num() > MeshMapping.MeshIndex);
@@ -45,7 +49,8 @@ namespace UE::NearestNeighborModel
 
 			if (!Track->GetMeshDataAtTime(SampleTime, GeomCacheMeshData))
 			{
-				UE_LOG(LogNearestNeighborModel, Error, TEXT("FNearestNeighborGeomCacheSampler::SamplePart: Track cannot get mesh delta at frame %d"), InAnimFrameIndex);
+				UE_LOG(LogNearestNeighborModel, Error, TEXT("SamplePart: Track cannot get mesh delta at frame %d"), InAnimFrameIndex);
+				return EUpdateResult::ERROR;
 			}
 
 			// Calculate the vertex deltas.
@@ -85,6 +90,7 @@ namespace UE::NearestNeighborModel
 				}
 			}
 		}
+		return EUpdateResult::SUCCESS;
 	}
 
 	bool IsPotentialMatch(const FString& TrackName, const FString& MeshName)
@@ -92,8 +98,9 @@ namespace UE::NearestNeighborModel
 		return (TrackName.Find(MeshName) == 0);
 	}
 
-	void FNearestNeighborGeomCacheSampler::GeneratePartMeshMappings(const TArray<uint32>& VertexMap, bool bUsePartOnlyMesh)
+	uint8 FNearestNeighborGeomCacheSampler::GeneratePartMeshMappings(const TArray<uint32>& VertexMap, bool bUsePartOnlyMesh)
 	{
+		uint8 Result = EUpdateResult::SUCCESS;
 		USkeletalMesh* SkeletalMesh = SkeletalMeshComponent.Get() ? SkeletalMeshComponent->GetSkeletalMeshAsset() : nullptr;
 		UGeometryCache* GeometryCache = GeometryCacheComponent.Get() ? GeometryCacheComponent->GetGeometryCache() : nullptr;
 		// TODO: make this more general
@@ -104,18 +111,25 @@ namespace UE::NearestNeighborModel
 				TArray<FString> FailedNames;
 				TArray<FString> VertexMisMatchNames;
 				GenerateGeomCacheMeshMappings(SkeletalMesh, GeometryCache, MeshMappings, FailedNames, VertexMisMatchNames);
-				return;
+				if (!FailedNames.IsEmpty() || !VertexMisMatchNames.IsEmpty())
+				{
+					Result |= EUpdateResult::WARNING;
+				}
+				for(int32 i = 0; i < VertexMisMatchNames.Num(); i++)
+				{
+					UE_LOG(LogNearestNeighborModel, Warning, TEXT("%s is skipped because it has different vertex counts in skeletal mesh and geometry cache."), *VertexMisMatchNames[i]);
+				}
+				return Result;
 			}
 			FSkeletalMeshModel* ImportedModel = SkeletalMesh->GetImportedModel();
-			check(ImportedModel);
-			check(ImportedModel->LODModels[0].ImportedMeshInfos.Num() >= 1);
 
-			const TArray<FSkelMeshImportedMeshInfo>& SkelMeshInfos = ImportedModel->LODModels[0].ImportedMeshInfos;
-			if (SkelMeshInfos.IsEmpty())
+			if (!ImportedModel || ImportedModel->LODModels[0].ImportedMeshInfos.IsEmpty())
 			{
-				return;
+				UE_LOG(LogNearestNeighborModel, Error, TEXT("Unable to generate mesh mappings because SkeletalMesh has no imported model."));
+				return EUpdateResult::ERROR;
 			}
 
+			const TArray<FSkelMeshImportedMeshInfo>& SkelMeshInfos = ImportedModel->LODModels[0].ImportedMeshInfos;
 			MeshMappings.Reset();
 
 			FString SkelMeshName;
@@ -151,6 +165,7 @@ namespace UE::NearestNeighborModel
 						// Verify that we have imported vertex numbers.
 						if (GeomCacheMeshData.ImportedVertexNumbers.IsEmpty())
 						{
+							UE_LOG(LogNearestNeighborModel, Warning, TEXT("Geometry track %s is skipped because it has no imported vertex numbers."), *Track->GetName());
 							continue;
 						}
 
@@ -173,7 +188,7 @@ namespace UE::NearestNeighborModel
 							Mapping.ImportedVertexToRenderVertexMap[PartVertexIndex] = RenderVertexIndex;
 						}
 
-						// We found a match, no need to iterate over more Tracks.
+						// We found a match, no need to iterate over more MeshInfos.
 						bFoundMatch = true;
 						break;
 					} // If the track name matches the skeletal meshes internal mesh name.
@@ -181,10 +196,17 @@ namespace UE::NearestNeighborModel
 
 				if (Track && !bFoundMatch)
 				{
-					UE_LOG(LogMLDeformer, Warning, TEXT("Geometry cache '%s' cannot be matched with a mesh inside the Skeletal Mesh."), *Track->GetName());
+					Result |= EUpdateResult::WARNING;
+					UE_LOG(LogNearestNeighborModel, Warning, TEXT("Geometry cache '%s' cannot be matched with a mesh inside the Skeletal Mesh."), *Track->GetName());
 				}
 			} // For all tracks.
 		}
+		else
+		{
+			Result |= EUpdateResult::WARNING;
+			UE_LOG(LogNearestNeighborModel, Warning, TEXT("SkeletalMesh or GeometryCache is none. No mapping is generated"));
+		}
+		return Result;
 	}
 
 	void FNearestNeighborGeomCacheSampler::SampleKMeansAnim(const int32 SkeletonId)
diff --git a/Engine/Plugins/Animation/MLDeformer/NearestNeighborModel/Source/NearestNeighborModelEditor/Private/NearestNeighborModelDetails.cpp b/Engine/Plugins/Animation/MLDeformer/NearestNeighborModel/Source/NearestNeighborModelEditor/Private/NearestNeighborModelDetails.cpp
--- a/Engine/Plugins/Animation/MLDeformer/NearestNeighborModel/Source/NearestNeighborModelEditor/Private/NearestNeighborModelDetails.cpp
+++ b/Engine/Plugins/Animation/MLDeformer/NearestNeighborModel/Source/NearestNeighborModelEditor/Private/NearestNeighborModelDetails.cpp
@@ -41,10 +41,12 @@ namespace UE::NearestNeighborModel
 
 	void FNearestNeighborModelDetails::CreateCategories()
 	{
-		FMLDeformerMorphModelDetails::CreateCategories();
+		FMLDeformerGeomCacheModelDetails::CreateCategories();
+
 		FileCacheCategoryBuilder = &DetailLayoutBuilder->EditCategory("File Cache", FText::GetEmpty(), ECategoryPriority::Important);
 		ClothPartCategoryBuilder = &DetailLayoutBuilder->EditCategory("Cloth Parts", FText::GetEmpty(), ECategoryPriority::Important);
 		NearestNeighborCategoryBuilder = &DetailLayoutBuilder->EditCategory("Nearest Neighbors", FText::GetEmpty(), ECategoryPriority::Important);
+		MorphTargetCategoryBuilder = &DetailLayoutBuilder->EditCategory("Morph Targets", FText::GetEmpty(), ECategoryPriority::Important);
 		KMeansCategoryBuilder = &DetailLayoutBuilder->EditCategory("KMeans Pose Generator", FText::GetEmpty(), ECategoryPriority::Important);
 
 		// Add warning in CreateCategories so that the warning appears at the top of the details panel.
@@ -64,7 +66,7 @@ namespace UE::NearestNeighborModel
 	void FNearestNeighborModelDetails::CustomizeDetails(IDetailLayoutBuilder& DetailBuilder)
 	{
 		// Create all the detail categories and add the properties of the base class.
-		FMLDeformerMorphModelDetails::CustomizeDetails(DetailBuilder);
+		FMLDeformerGeomCacheModelDetails::CustomizeDetails(DetailBuilder);
 
 		// Training settings.
 		TrainingSettingsCategoryBuilder->AddProperty(UNearestNeighborModel::GetInputDimPropertyName());
@@ -108,44 +110,11 @@ namespace UE::NearestNeighborModel
 				.Text(VertexCountText)
 			];
 
-
-		FText ButtonText = (NearestNeighborModel && NearestNeighborModel->IsClothPartDataValid()) ? LOCTEXT("Update", "Update") : LOCTEXT("Update *", "Update *");
-		ClothPartCategoryBuilder->AddCustomRow(FText::FromString(""))
-			.WholeRowContent()
-			[
-				SNew(SButton)
-				.Text(ButtonText)
-				.HAlign(HAlign_Center)
-				.OnClicked_Lambda([this]
-				{
-					NearestNeighborModel->UpdateClothPartData();
-					NearestNeighborEditorModel->UpdateNearestNeighborActors();
-					EditorModel->GetEditor()->GetModelDetailsView()->ForceRefresh();
-					return FReply::Handled();
-				})
-			];
-
 		// Nearest Neighbor settings
 		NearestNeighborCategoryBuilder->AddProperty(GET_MEMBER_NAME_STRING_CHECKED(UNearestNeighborModel, DecayFactor));
 		NearestNeighborCategoryBuilder->AddProperty(GET_MEMBER_NAME_STRING_CHECKED(UNearestNeighborModel, NearestNeighborOffsetWeight));
 		NearestNeighborCategoryBuilder->AddProperty(UNearestNeighborModel::GetUsePartOnlyMeshPropertyName());
 		NearestNeighborCategoryBuilder->AddProperty(UNearestNeighborModel::GetNearestNeighborDataPropertyName());
-		ButtonText = NearestNeighborModel->IsNearestNeighborDataValid() ? LOCTEXT("Update", "Update") : LOCTEXT("Update *", "Update *");
-
-		NearestNeighborCategoryBuilder->AddCustomRow(FText::FromString(""))
-			.WholeRowContent()
-			[
-				SNew(SButton)
-				.Text(ButtonText)
-				.HAlign(HAlign_Center)
-				.OnClicked_Lambda([this]
-				{
-					NearestNeighborEditorModel->UpdateNearestNeighborData();
-					NearestNeighborEditorModel->UpdateNearestNeighborActors();
-					EditorModel->GetEditor()->GetModelDetailsView()->ForceRefresh();
-					return FReply::Handled();
-				})
-			];
 
 		KMeansCategoryBuilder->AddProperty(GET_MEMBER_NAME_STRING_CHECKED(UNearestNeighborModel, SourceSkeletons));
 		KMeansCategoryBuilder->AddProperty(GET_MEMBER_NAME_STRING_CHECKED(UNearestNeighborModel, NumClusters));
@@ -162,17 +131,20 @@ namespace UE::NearestNeighborModel
 				})
 			];
 		
+		MorphTargetCategoryBuilder->AddProperty(UMLDeformerMorphModel::GetMorphTargetDeltaThresholdPropertyName(), UMLDeformerMorphModel::StaticClass());
+		MorphTargetCategoryBuilder->AddProperty(UMLDeformerMorphModel::GetMorphTargetErrorTolerancePropertyName(), UMLDeformerMorphModel::StaticClass());
+		FText ButtonText = NearestNeighborModel->IsMorphTargetDataValid() ? LOCTEXT("Update", "Update") : LOCTEXT("Update *", "Update *");
 		MorphTargetCategoryBuilder->AddProperty(UNearestNeighborModel::GetMorphDataSizePropertyName());
 		MorphTargetCategoryBuilder->AddCustomRow(FText::FromString(""))
 			.WholeRowContent()
 			[
 				SNew(SButton)
-				.Text(FText::FromString("Update"))
+				.Text(ButtonText)
 				.HAlign(HAlign_Center)
 				.OnClicked_Lambda([this]
 				{
 					NearestNeighborEditorModel->OnMorphTargetUpdate();
-					NearestNeighborEditorModel->InitTestMLDeformerPreviousWeights();
+					EditorModel->GetEditor()->GetModelDetailsView()->ForceRefresh();
 					return FReply::Handled();
 				})
 			];
@@ -182,7 +154,7 @@ namespace UE::NearestNeighborModel
 	void FNearestNeighborModelDetails::AddUpdateResultText(IDetailCategoryBuilder* CategoryBuilder, uint8 Result)
 	{
 		CategoryBuilder->AddCustomRow(FText::FromString("UpdateResultError"))
-			.Visibility((Result & ERROR) != 0 ? EVisibility::Visible : EVisibility::Collapsed)
+			.Visibility((Result & EUpdateResult::ERROR) != 0 ? EVisibility::Visible : EVisibility::Collapsed)
 			.WholeRowContent()
 			[
 				SNew(SBox)
@@ -194,7 +166,7 @@ namespace UE::NearestNeighborModel
 				]
 			];
 		CategoryBuilder->AddCustomRow(FText::FromString("UpdateResultWarning"))
-			.Visibility((Result & WARNING) != 0 ? EVisibility::Visible : EVisibility::Collapsed)
+			.Visibility((Result & EUpdateResult::WARNING) != 0 ? EVisibility::Visible : EVisibility::Collapsed)
 			.WholeRowContent()
 			[
 				SNew(SBox)
@@ -205,7 +177,7 @@ namespace UE::NearestNeighborModel
 					.Message(FText::FromString(TEXT("Update finished with warnings. Please check Output Log (LogNearestNeighborModel, LogPython) for details.")))
 				]
 			];
-		if (Result == SUCCESS)
+		if (Result == EUpdateResult::SUCCESS)
 		{
 			UE_LOG(LogNearestNeighborModel, Display, TEXT("Update succeeded."));
 		}
diff --git a/Engine/Plugins/Animation/MLDeformer/NearestNeighborModel/Source/NearestNeighborModelEditor/Private/NearestNeighborTrainingModel.cpp b/Engine/Plugins/Animation/MLDeformer/NearestNeighborModel/Source/NearestNeighborModelEditor/Private/NearestNeighborTrainingModel.cpp
--- a/Engine/Plugins/Animation/MLDeformer/NearestNeighborModel/Source/NearestNeighborModelEditor/Private/NearestNeighborTrainingModel.cpp
+++ b/Engine/Plugins/Animation/MLDeformer/NearestNeighborModel/Source/NearestNeighborModelEditor/Private/NearestNeighborTrainingModel.cpp
@@ -33,22 +33,23 @@ const TArray<int32> UNearestNeighborTrainingModel::GetPartVertexMap(const int32
 	return TArray<int32>((int32*)VertexMap.GetData(), VertexMap.Num());
 }
 
-void UNearestNeighborTrainingModel::SamplePart(int32 PartId, int32 Index)
+int32 UNearestNeighborTrainingModel::SamplePart(int32 PartId, int32 Index)
 {
 	FNearestNeighborGeomCacheSampler* Sampler = static_cast<FNearestNeighborGeomCacheSampler*>(EditorModel->GetSampler());
-	Sampler->SamplePart(Index, NearestNeighborModel->PartVertexMap(PartId));
+	int32 Result = Sampler->SamplePart(Index, NearestNeighborModel->PartVertexMap(PartId));
 	PartSampleDeltas = Sampler->GetPartVertexDeltas();
 	SampleBoneRotations = Sampler->GetBoneRotations();
+	return Result;
 }
 
-void UNearestNeighborTrainingModel::SetSamplerPartData(const int32 PartId)
+int32 UNearestNeighborTrainingModel::SetSamplerPartData(const int32 PartId)
 {
-	GetNearestNeighborEditorModel()->SetSamplerPartData(PartId);
+	return GetNearestNeighborEditorModel()->SetSamplerPartData(PartId);
 }
 
 int32 UNearestNeighborTrainingModel::GetPartNumNeighbors(const int32 PartId) const
 {
-	return NearestNeighborModel->GetNumNeighborsFromGeometryCache(PartId);
+	return FMath::Min(NearestNeighborModel->GetNumNeighborsFromAnimSequence(PartId), NearestNeighborModel->GetNumNeighborsFromGeometryCache(PartId));
 }
 
 void UNearestNeighborTrainingModel::SampleKmeansAnim(const int32 SkeletonId)
diff --git a/Engine/Plugins/Animation/MLDeformer/NearestNeighborModel/Source/NearestNeighborModelEditor/Public/NearestNeighborEditorModel.h b/Engine/Plugins/Animation/MLDeformer/NearestNeighborModel/Source/NearestNeighborModelEditor/Public/NearestNeighborEditorModel.h
--- a/Engine/Plugins/Animation/MLDeformer/NearestNeighborModel/Source/NearestNeighborModelEditor/Public/NearestNeighborEditorModel.h
+++ b/Engine/Plugins/Animation/MLDeformer/NearestNeighborModel/Source/NearestNeighborModelEditor/Public/NearestNeighborEditorModel.h
@@ -50,10 +50,10 @@ namespace UE::NearestNeighborModel
 		UNearestNeighborModel* GetNearestNeighborModel() const { return static_cast<UNearestNeighborModel*>(Model); }
 	
 		// Recomputes nearest neighbor coeffcients and nearest neighbor vertex offsets. 
-		virtual void UpdateNearestNeighborData();
+		virtual uint8 UpdateNearestNeighborData();
 
 		// Temporarily set sampler to use part data in nearest neighbor dataset. This is useful when updating nearest neighbor data.
-		bool SetSamplerPartData(const int32 PartI);
+		uint8 SetSamplerPartData(const int32 PartI);
 
 		// Reset smapler to use the original dataset.
 		void ResetSamplerData();
@@ -71,6 +71,8 @@ namespace UE::NearestNeighborModel
 		uint8 GetMorphTargetUpdateResult() { return MorphTargetUpdateResult; }
 		UMLDeformerComponent* GetTestMLDeformerComponent();
 		void InitTestMLDeformerPreviousWeights();
+		uint8 WarnIfNetworkInvalid();
+		bool IsNeuralNetworkLoaded();
 
 
 	protected:
diff --git a/Engine/Plugins/Animation/MLDeformer/NearestNeighborModel/Source/NearestNeighborModelEditor/Public/NearestNeighborGeomCacheSampler.h b/Engine/Plugins/Animation/MLDeformer/NearestNeighborModel/Source/NearestNeighborModelEditor/Public/NearestNeighborGeomCacheSampler.h
--- a/Engine/Plugins/Animation/MLDeformer/NearestNeighborModel/Source/NearestNeighborModelEditor/Public/NearestNeighborGeomCacheSampler.h
+++ b/Engine/Plugins/Animation/MLDeformer/NearestNeighborModel/Source/NearestNeighborModelEditor/Public/NearestNeighborGeomCacheSampler.h
@@ -10,11 +10,12 @@ namespace UE::NearestNeighborModel
 		: public UE::MLDeformer::FMLDeformerGeomCacheSampler
 	{
 	public:
-		virtual void SamplePart(int32 InAnimFrameIndex, const TArray<uint32>& VertexMap);
+		virtual uint8 SamplePart(int32 InAnimFrameIndex, const TArray<uint32>& VertexMap);
 		virtual void SampleKMeansAnim(const int32 SkeletonId);
 		virtual void SampleKMeansFrame(const int32 Frame);
 		const TArray<float>& GetPartVertexDeltas() const { return PartVertexDeltas; }
-		void GeneratePartMeshMappings(const TArray<uint32>& VertexMap, bool bUsePartOnlyMesh);
+		uint8 GeneratePartMeshMappings(const TArray<uint32>& VertexMap, bool bUsePartOnlyMesh);
+		bool IsMeshMappingsEmpty() const { return MeshMappings.IsEmpty(); }
 	protected:
 		TArray<float> PartVertexDeltas;
 	};
diff --git a/Engine/Plugins/Animation/MLDeformer/NearestNeighborModel/Source/NearestNeighborModelEditor/Public/NearestNeighborTrainingModel.h b/Engine/Plugins/Animation/MLDeformer/NearestNeighborModel/Source/NearestNeighborModelEditor/Public/NearestNeighborTrainingModel.h
--- a/Engine/Plugins/Animation/MLDeformer/NearestNeighborModel/Source/NearestNeighborModelEditor/Public/NearestNeighborTrainingModel.h
+++ b/Engine/Plugins/Animation/MLDeformer/NearestNeighborModel/Source/NearestNeighborModelEditor/Public/NearestNeighborTrainingModel.h
@@ -35,13 +35,13 @@ public:
 	const TArray<int32> GetPartVertexMap(const int32 PartId) const;
 
 	UFUNCTION(BlueprintCallable, Category = "Python")
-	void SamplePart(int32 PartId, int32 Index);
+	int32 SamplePart(int32 PartId, int32 Index);
 
 	UFUNCTION(BlueprintImplementableEvent, Category = "Python")
-	void UpdateNearestNeighborData();
+	int32 UpdateNearestNeighborData();
 
 	UFUNCTION(BlueprintCallable, Category = "Python")
-	void SetSamplerPartData(const int32 PartId);
+	int32 SetSamplerPartData(const int32 PartId);
 
 	UFUNCTION(BlueprintPure, Category = "Python")
 	int32 GetPartNumNeighbors(const int32 PartId) const;
