NearestNeighborEditor: resize GeomCacheMeshDatas and fix error message.

diff --git a/Engine/Plugins/Animation/MLDeformer/NearestNeighborModel/Source/NearestNeighborModelEditor/Private/NearestNeighborEditorModel.cpp b/Engine/Plugins/Animation/MLDeformer/NearestNeighborModel/Source/NearestNeighborModelEditor/Private/NearestNeighborEditorModel.cpp
--- a/Engine/Plugins/Animation/MLDeformer/NearestNeighborModel/Source/NearestNeighborModelEditor/Private/NearestNeighborEditorModel.cpp
+++ b/Engine/Plugins/Animation/MLDeformer/NearestNeighborModel/Source/NearestNeighborModelEditor/Private/NearestNeighborEditorModel.cpp
@@ -254,7 +254,7 @@ namespace UE::NearestNeighborModel
 					}
 					else
 					{
-						UE_LOG(LogNearestNeighborModel, Warning, TEXT("NearestNeighborData: part %d frame mismatch: AnimSequence has %d frames and GeometryCache has %d frames. Using %f frames only."), PartId, NumNeighborsFromAnimSequence, NumNeighborsFromGeomCache, NumFrames);
+						UE_LOG(LogNearestNeighborModel, Warning, TEXT("NearestNeighborData: part %d frame mismatch: AnimSequence has %d frames and GeometryCache has %d frames. Using %d frames only."), PartId, NumNeighborsFromAnimSequence, NumNeighborsFromGeomCache, NumFrames);
 						return EUpdateResult::WARNING;
 					}
 				}
diff --git a/Engine/Plugins/Animation/MLDeformer/NearestNeighborModel/Source/NearestNeighborModelEditor/Private/NearestNeighborGeomCacheSampler.cpp b/Engine/Plugins/Animation/MLDeformer/NearestNeighborModel/Source/NearestNeighborModelEditor/Private/NearestNeighborGeomCacheSampler.cpp
--- a/Engine/Plugins/Animation/MLDeformer/NearestNeighborModel/Source/NearestNeighborModelEditor/Private/NearestNeighborGeomCacheSampler.cpp
+++ b/Engine/Plugins/Animation/MLDeformer/NearestNeighborModel/Source/NearestNeighborModelEditor/Private/NearestNeighborGeomCacheSampler.cpp
@@ -119,6 +119,13 @@ namespace UE::NearestNeighborModel
 				{
 					UE_LOG(LogNearestNeighborModel, Warning, TEXT("%s is skipped because it has different vertex counts in skeletal mesh and geometry cache."), *VertexMisMatchNames[i]);
 				}
+				GeomCacheMeshDatas.Reset();
+				GeomCacheMeshDatas.AddDefaulted(MeshMappings.Num());
+				if(MeshMappings.IsEmpty())
+				{
+					UE_LOG(LogNearestNeighborModel, Error, TEXT("Unable to match skeletal mesh with geometry cache."));
+					Result = EUpdateResult::ERROR;
+				}
 				return Result;
 			}
 			FSkeletalMeshModel* ImportedModel = SkeletalMesh->GetImportedModel();
