NearestNeighborEditor: Add drop menus and make txt file optional

diff --git a/Engine/Plugins/Animation/MLDeformer/NearestNeighborModel/Source/NearestNeighborModel/Private/NearestNeighborModel.cpp b/Engine/Plugins/Animation/MLDeformer/NearestNeighborModel/Source/NearestNeighborModel/Private/NearestNeighborModel.cpp
--- a/Engine/Plugins/Animation/MLDeformer/NearestNeighborModel/Source/NearestNeighborModel/Private/NearestNeighborModel.cpp
+++ b/Engine/Plugins/Animation/MLDeformer/NearestNeighborModel/Source/NearestNeighborModel/Private/NearestNeighborModel.cpp
@@ -78,6 +78,26 @@ TArray<uint32> ReadTxt(const FString &Path)
 	return UIntArr;
 }
 
+TArray<uint32> Range(const uint32 Start, const uint32 End)
+{
+	TArray<uint32> Arr; Arr.SetNum(End - Start + 1);
+	for (uint32 i = Start; i < End; i++)
+	{
+		Arr[i - Start] = i;
+	}
+	return Arr;
+}
+
+TArray<uint32> AddConstant(const TArray<uint32> &InArr, uint32 Constant)
+{
+	TArray<uint32> OutArr = InArr;
+	for (int32 i = 0; i < InArr.Num(); i++)
+	{
+		OutArr[i] = InArr[i] + Constant;
+	}
+	return OutArr;
+}
+
 void UNearestNeighborModel::ClipInputs(float* InputPtr, int NumInputs)
 {
 	if(InputsMin.Num() == NumInputs && InputsMax.Num() == NumInputs)
@@ -157,40 +177,86 @@ void UNearestNeighborModel::UpdateNetworkOutputDim()
 	}
 }
 
+UE::NearestNeighborModel::EUpdateResult UNearestNeighborModel::UpdateVertexMap(int32 PartId, const FString& VertexMapPath, const FSkelMeshImportedMeshInfo& Info)
+{
+	using namespace UE::NearestNeighborModel;
+	uint8 ReturnCode = EUpdateResult::SUCCESS;
+	const int32 StartIndex = Info.StartImportedVertex;
+	const int32 NumVertices = Info.NumVertices;
+	bool bIsVertexMapValid = true;
+
+	if (VertexMapPath.IsEmpty())
+	{
+		bIsVertexMapValid = false;
+	}
+	else if (!FPaths::FileExists(VertexMapPath))
+	{
+		bIsVertexMapValid = false;
+		ReturnCode |= EUpdateResult::ERROR;
+		UE_LOG(LogNearestNeighborModel, Error, TEXT("Part %d txt path %s does not exist"), PartId, *VertexMapPath, NumVertices);
+	}
+	else
+	{
+		const TArray<uint32> PartVertexMap = ReadTxt(ClothPartEditorData[PartId].VertexMapPath);
+		if (PartVertexMap.Num() > NumVertices)
+		{
+			bIsVertexMapValid = false;
+			UE_LOG(LogNearestNeighborModel, Error, TEXT("Part %d vertex map has %d vertices, larger than %d vertices in skeletal mesh, using %d vertices instead"), PartId, PartVertexMap.Num(), NumVertices, NumVertices);
+			ReturnCode |= EUpdateResult::ERROR;
+		}
+		else
+		{
+			const int32 MaxIndex = FMath::Max(PartVertexMap);
+			if (MaxIndex >= NumVertices)
+			{
+				bIsVertexMapValid = false;
+				UE_LOG(LogNearestNeighborModel, Error, TEXT("Part %d vertex map max index is %d. There are only %d vertices in skeletal mesh, using %d vertices instead"), PartId, MaxIndex, NumVertices, NumVertices);
+				ReturnCode |= EUpdateResult::ERROR;
+			}
+		}
+		if (bIsVertexMapValid)
+		{
+			ClothPartData[PartId].VertexMap = AddConstant(PartVertexMap, StartIndex);
+			return (EUpdateResult)ReturnCode;
+		}
+	}
+
+	ClothPartData[PartId].VertexMap = Range(StartIndex, StartIndex + NumVertices);
+	return (EUpdateResult)ReturnCode;
+}
+
 UE::NearestNeighborModel::EUpdateResult UNearestNeighborModel::UpdateClothPartData()
 {
 	using namespace UE::NearestNeighborModel;
+	uint8 ReturnCode = EUpdateResult::SUCCESS;
 	if(ClothPartEditorData.Num() == 0)
 	{
 		UE_LOG(LogNearestNeighborModel, Error, TEXT("There should be at least 1 cloth part"));
 		return EUpdateResult::ERROR;
 	}
+	if (!GetSkeletalMesh() || !GetSkeletalMesh()->GetImportedModel() || GetSkeletalMesh()->GetImportedModel()->LODModels[0].ImportedMeshInfos.IsEmpty())
+	{
+		UE_LOG(LogNearestNeighborModel, Error, TEXT("SkeletalMesh is None or SkeletalMesh has no imported model."));
+		return EUpdateResult::ERROR;
+	}
+	const TArray<FSkelMeshImportedMeshInfo>& SkelMeshInfos = GetSkeletalMesh()->GetImportedModel()->LODModels[0].ImportedMeshInfos;
 
 	ClothPartData.SetNum(ClothPartEditorData.Num());
 	for(int32 PartId = 0; PartId < GetNumParts(); PartId++)
 	{
-		if(FPaths::FileExists(ClothPartEditorData[PartId].VertexMapPath))
-		{
-			ClothPartData[PartId].PCACoeffNum = ClothPartEditorData[PartId].PCACoeffNum;
-			ClothPartData[PartId].VertexMap = ReadTxt(ClothPartEditorData[PartId].VertexMapPath);
-			ClothPartData[PartId].NumVertices = ClothPartData[PartId].VertexMap.Num();
-
-			// Initialize PCA data with 0 to prevent NearestNeighborModelInstance::Tick from crashing
-			if (!CheckPCAData(PartId))
-			{
-				ClothPartData[PartId].VertexMean.SetNumZeroed(ClothPartData[PartId].NumVertices * 3);
-				ClothPartData[PartId].PCABasis.SetNumZeroed(ClothPartData[PartId].NumVertices * 3 * ClothPartData[PartId].PCACoeffNum);
+		ClothPartData[PartId].PCACoeffNum = ClothPartEditorData[PartId].PCACoeffNum;
+		ReturnCode |= UpdateVertexMap(PartId, ClothPartEditorData[PartId].VertexMapPath, SkelMeshInfos[ClothPartEditorData[PartId].MeshIndex]);
+		ClothPartData[PartId].NumVertices = ClothPartData[PartId].VertexMap.Num();
 
-				// Init default neighbor data.
-				ClothPartData[PartId].NeighborCoeffs.SetNumZeroed(ClothPartData[PartId].PCACoeffNum);
-				ClothPartData[PartId].NeighborOffsets.SetNumZeroed(ClothPartData[PartId].NumVertices * 3);
-				ClothPartData[PartId].NumNeighbors = 1;
-			}
-		}
-		else
+		if (!CheckPCAData(PartId))
 		{
-			UE_LOG(LogNearestNeighborModel, Error, TEXT("Part %d txt path %s does not exist"), PartId, *ClothPartEditorData[PartId].VertexMapPath);
-			return EUpdateResult::ERROR;
+			ClothPartData[PartId].VertexMean.SetNumZeroed(ClothPartData[PartId].NumVertices * 3);
+			ClothPartData[PartId].PCABasis.SetNumZeroed(ClothPartData[PartId].NumVertices * 3 * ClothPartData[PartId].PCACoeffNum);
+
+			// Init default neighbor data.
+			ClothPartData[PartId].NeighborCoeffs.SetNumZeroed(ClothPartData[PartId].PCACoeffNum);
+			ClothPartData[PartId].NeighborOffsets.SetNumZeroed(ClothPartData[PartId].NumVertices * 3);
+			ClothPartData[PartId].NumNeighbors = 1;
 		}
 	}
 	UpdatePCACoeffNums();
@@ -199,7 +265,7 @@ UE::NearestNeighborModel::EUpdateResult UNearestNeighborModel::UpdateClothPartDa
 	UpdateNetworkInputDim();
 	UpdateNetworkOutputDim();
 	bClothPartDataValid = true;
-	return EUpdateResult::SUCCESS;
+	return (EUpdateResult)ReturnCode;
 }
 
 TObjectPtr<UAnimSequence> UNearestNeighborModel::GetNearestNeighborSkeletons(int32 PartId)
diff --git a/Engine/Plugins/Animation/MLDeformer/NearestNeighborModel/Source/NearestNeighborModel/Public/NearestNeighborModel.h b/Engine/Plugins/Animation/MLDeformer/NearestNeighborModel/Source/NearestNeighborModel/Public/NearestNeighborModel.h
--- a/Engine/Plugins/Animation/MLDeformer/NearestNeighborModel/Source/NearestNeighborModel/Public/NearestNeighborModel.h
+++ b/Engine/Plugins/Animation/MLDeformer/NearestNeighborModel/Source/NearestNeighborModel/Public/NearestNeighborModel.h
@@ -10,6 +10,7 @@
 #include "UObject/Object.h"
 #include "NearestNeighborModel.generated.h"
 
+struct FSkelMeshImportedMeshInfo;
 class USkeletalMesh;
 class UGeometryCache;
 class UAnimSequence;
@@ -69,8 +70,11 @@ struct FClothPartEditorData
 	UPROPERTY(EditAnywhere, BlueprintReadWrite, meta = (ClampMin = "1", ClampMax = "128"), Category = "Nearest Neighbors")
 	int32 PCACoeffNum = 128;
 
-	UPROPERTY(EditAnywhere, BlueprintReadWrite, Category = "Nearest Neighbors")
+	UPROPERTY(EditAnywhere, BlueprintReadWrite, Category = "Nearest Neighbors", meta = (DisplayName = "Vertex Map Path (Optional)"))
 	FString VertexMapPath;
+
+	UPROPERTY(EditAnywhere, BlueprintReadWrite, Category = "Nearest Neighbors")
+	int32 MeshIndex = 0;
 };
 
 USTRUCT(BlueprintType, Blueprintable)
@@ -231,6 +235,7 @@ public:
 	void UpdateNetworkInputDim();
 	void UpdateNetworkOutputDim();
 	UE::NearestNeighborModel::EUpdateResult UpdateClothPartData();
+	UE::NearestNeighborModel::EUpdateResult UpdateVertexMap(int32 PartId, const FString& VertexMapPath, const FSkelMeshImportedMeshInfo& Info);
 	void UpdatePCACoeffNums();
 	void UpdateNetworkSize();
 	void UpdateMorphTargetSize();
@@ -296,6 +301,9 @@ public:
 
 	UFUNCTION(BlueprintPure, Category = "Nearest Neighbor Model")
 	FString GetModelDir() const;
+
+	int32 GetPartMeshIndex(int32 PartId) const { return ClothPartEditorData[PartId].MeshIndex; }
+	void SetPartMeshIndex(int32 PartId, int32 MeshIndex) { ClothPartEditorData[PartId].MeshIndex = MeshIndex; }
 #endif
 
 
diff --git a/Engine/Plugins/Animation/MLDeformer/NearestNeighborModel/Source/NearestNeighborModelEditor/NearestNeighborModelEditor.Build.cs b/Engine/Plugins/Animation/MLDeformer/NearestNeighborModel/Source/NearestNeighborModelEditor/NearestNeighborModelEditor.Build.cs
--- a/Engine/Plugins/Animation/MLDeformer/NearestNeighborModel/Source/NearestNeighborModelEditor/NearestNeighborModelEditor.Build.cs
+++ b/Engine/Plugins/Animation/MLDeformer/NearestNeighborModel/Source/NearestNeighborModelEditor/NearestNeighborModelEditor.Build.cs
@@ -22,6 +22,7 @@ namespace UnrealBuildTool.Rules
 					"UnrealEd",
 					"CoreUObject",
 					"Engine",
+					"InputCore",
 					"Slate",
 					"SlateCore",
 					"EditorWidgets",
diff --git a/Engine/Plugins/Animation/MLDeformer/NearestNeighborModel/Source/NearestNeighborModelEditor/Private/NearestNeighborGeomCacheSampler.cpp b/Engine/Plugins/Animation/MLDeformer/NearestNeighborModel/Source/NearestNeighborModelEditor/Private/NearestNeighborGeomCacheSampler.cpp
--- a/Engine/Plugins/Animation/MLDeformer/NearestNeighborModel/Source/NearestNeighborModelEditor/Private/NearestNeighborGeomCacheSampler.cpp
+++ b/Engine/Plugins/Animation/MLDeformer/NearestNeighborModel/Source/NearestNeighborModelEditor/Private/NearestNeighborGeomCacheSampler.cpp
@@ -17,7 +17,7 @@
 using namespace UE::MLDeformer;
 namespace UE::NearestNeighborModel
 {
-	uint8 FNearestNeighborGeomCacheSampler::SamplePart(int32 InAnimFrameIndex, const TArray<uint32>& VertexMap)
+	uint8 FNearestNeighborGeomCacheSampler::SamplePart(int32 InAnimFrameIndex, int32 PartId)
 	{
 		FMLDeformerSampler::Sample(InAnimFrameIndex);
 		USkeletalMesh* SkeletalMesh = SkeletalMeshComponent.Get() ? SkeletalMeshComponent->GetSkeletalMeshAsset() : nullptr;
@@ -35,7 +35,12 @@ namespace UE::NearestNeighborModel
 				UE_LOG(LogNearestNeighborModel, Error, TEXT("SamplePart: MeshMappings is empty."));
 				return EUpdateResult::ERROR;
 			}
-			const UE::MLDeformer::FMLDeformerGeomCacheMeshMapping& MeshMapping = MeshMappings[0]; 
+			if (PartId >= MeshMappingIndices.Num())
+			{
+				UE_LOG(LogNearestNeighborModel, Error, TEXT("SamplePart: MeshMappingIndices.Num()=%d is smaller than PartId %d"), MeshMappingIndices.Num(), PartId);
+				return EUpdateResult::ERROR;
+			}
+			const UE::MLDeformer::FMLDeformerGeomCacheMeshMapping& MeshMapping = MeshMappings[MeshMappingIndices[PartId]]; 
 
 			check(SkelMeshInfos.Num() > MeshMapping.MeshIndex);
 			const FSkelMeshImportedMeshInfo& MeshInfo = SkelMeshInfos[MeshMapping.MeshIndex]; 
@@ -57,6 +62,9 @@ namespace UE::NearestNeighborModel
 			const FSkeletalMeshLODRenderData& SkelMeshLODData = SkeletalMesh->GetResourceForRendering()->LODRenderData[LODIndex];
 			const FSkinWeightVertexBuffer& SkinWeightBuffer = *SkeletalMeshComponent->GetSkinWeightBuffer(LODIndex);
 
+			UNearestNeighborModel* NearestNeighborModel = static_cast<UNearestNeighborModel*>(Model);
+			check(NearestNeighborModel != nullptr);
+			const TArray<uint32>& VertexMap = NearestNeighborModel->PartVertexMap(PartId);
 			const int32 NumPartVerts = VertexMap.Num();
 			PartVertexDeltas.Reset();
 			PartVertexDeltas.SetNum(NumPartVerts * 3);
@@ -103,7 +111,6 @@ namespace UE::NearestNeighborModel
 		uint8 Result = EUpdateResult::SUCCESS;
 		USkeletalMesh* SkeletalMesh = SkeletalMeshComponent.Get() ? SkeletalMeshComponent->GetSkeletalMeshAsset() : nullptr;
 		UGeometryCache* GeometryCache = GeometryCacheComponent.Get() ? GeometryCacheComponent->GetGeometryCache() : nullptr;
-		// TODO: make this more general
 		if (SkeletalMeshComponent && SkeletalMesh && GeometryCacheComponent && GeometryCache)
 		{
 			if (!bUsePartOnlyMesh)
@@ -111,6 +118,7 @@ namespace UE::NearestNeighborModel
 				TArray<FString> FailedNames;
 				TArray<FString> VertexMisMatchNames;
 				GenerateGeomCacheMeshMappings(SkeletalMesh, GeometryCache, MeshMappings, FailedNames, VertexMisMatchNames);
+				Result |= GenerateMeshMappingIndices();
 				if (!FailedNames.IsEmpty() || !VertexMisMatchNames.IsEmpty())
 				{
 					Result |= EUpdateResult::WARNING;
@@ -152,12 +160,6 @@ namespace UE::NearestNeighborModel
 				{
 					const FSkelMeshImportedMeshInfo& MeshInfo = SkelMeshInfos[SkelMeshIndex];
 
-					// Hack for now
-					if (MeshInfo.StartImportedVertex != 0)
-					{
-						continue;
-					}
-
 					SkelMeshName = MeshInfo.Name.ToString();
 					if (Track &&
 						(IsPotentialMatch(Track->GetName(), SkelMeshName) || bIsSoloMesh))
@@ -213,6 +215,7 @@ namespace UE::NearestNeighborModel
 			Result |= EUpdateResult::WARNING;
 			UE_LOG(LogNearestNeighborModel, Warning, TEXT("SkeletalMesh or GeometryCache is none. No mapping is generated"));
 		}
+		Result |= GenerateMeshMappingIndices();
 		return Result;
 	}
 
@@ -244,4 +247,33 @@ namespace UE::NearestNeighborModel
 			UpdateCurveValues();
 		}
 	}
+
+	uint8 FNearestNeighborGeomCacheSampler::GenerateMeshMappingIndices()
+	{
+		UNearestNeighborModel* NearestNeighborModel = static_cast<UNearestNeighborModel*>(Model);
+		if (NearestNeighborModel)
+		{
+			const int32 NumParts = NearestNeighborModel->GetNumParts();
+			MeshMappingIndices.SetNum(NumParts);
+			for (int32 PartId = 0; PartId < NumParts; PartId++)
+			{
+				bool bFoundIndex = false;
+				for (int32 MappingId = 0; MappingId < MeshMappings.Num(); MappingId++)
+				{
+					if (MeshMappings[MappingId].MeshIndex == NearestNeighborModel->GetPartMeshIndex(PartId))
+					{
+						MeshMappingIndices[PartId] = MappingId;
+						bFoundIndex = true;
+						break;
+					}
+				}
+				if (!bFoundIndex)
+				{
+					UE_LOG(LogNearestNeighborModel, Error, TEXT("Part %d could not find a mesh mapping."), PartId);
+					return EUpdateResult::ERROR;
+				}
+			}
+		}
+		return EUpdateResult::SUCCESS;
+	} 
 };
diff --git a/Engine/Plugins/Animation/MLDeformer/NearestNeighborModel/Source/NearestNeighborModelEditor/Private/NearestNeighborModelDetails.cpp b/Engine/Plugins/Animation/MLDeformer/NearestNeighborModel/Source/NearestNeighborModelEditor/Private/NearestNeighborModelDetails.cpp
--- a/Engine/Plugins/Animation/MLDeformer/NearestNeighborModel/Source/NearestNeighborModelEditor/Private/NearestNeighborModelDetails.cpp
+++ b/Engine/Plugins/Animation/MLDeformer/NearestNeighborModel/Source/NearestNeighborModelEditor/Private/NearestNeighborModelDetails.cpp
@@ -7,10 +7,17 @@
 #include "DetailLayoutBuilder.h"
 #include "DetailCategoryBuilder.h"
 #include "DetailWidgetRow.h"
-#include "SWarningOrErrorBox.h"
+#include "Engine/SkeletalMesh.h"
+#include "IDetailChildrenBuilder.h"
 #include "IDetailGroup.h"
 #include "IDetailsView.h"
+#include "PropertyCustomizationHelpers.h"
+#include "Rendering/SkeletalMeshModel.h"
+#include "Rendering/SkeletalMeshLODModel.h"
+#include "SWarningOrErrorBox.h"
 #include "Widgets/Input/SButton.h"
+#include "Widgets/Input/SComboBox.h"
+#include "Widgets/Input/STextComboBox.h"
 #include "Widgets/Layout/SBox.h"
 #include "Widgets/Text/STextBlock.h"
 
@@ -63,6 +70,106 @@ namespace UE::NearestNeighborModel
 			];
 	}
 
+
+	void FNearestNeighborModelDetails::GenerateClothPartElementWidget(TSharedRef<IPropertyHandle> PropertyHandle, int32 ArrayIndex, IDetailChildrenBuilder& ChildrenBuilder)
+	{
+		TSharedPtr<IPropertyHandle> PCACoeffNumPropertyHandle = PropertyHandle->GetChildHandle(TEXT("PCACoeffNum"));
+		TSharedPtr<IPropertyHandle> VertexMapPathHandle = PropertyHandle->GetChildHandle(TEXT("VertexMapPath"));
+
+		typename STextComboBox::FOnTextSelectionChanged SubMeshComboOnSelectionChangedDelegate;
+		SubMeshComboOnSelectionChangedDelegate.BindRaw(this, &FNearestNeighborModelDetails::SubMeshComboSelectionChanged, ArrayIndex);
+		const int32 InitMeshIndex = NearestNeighborModel ? NearestNeighborModel->GetPartMeshIndex(ArrayIndex) : 0;
+
+		ChildrenBuilder.AddCustomRow(FText::GetEmpty())
+		[
+			SNew(SVerticalBox)
+			+SVerticalBox::Slot()
+			.AutoHeight()
+			[
+				SNew(SHorizontalBox)
+
+				+SHorizontalBox::Slot()
+				.Padding(5, 2)
+				.FillWidth(0.3f)
+				[
+					PropertyHandle->CreatePropertyNameWidget()
+				]
+
+				+SHorizontalBox::Slot()
+				.Padding(5, 2)
+				.FillWidth(0.7f)
+				[
+					PropertyHandle->CreatePropertyValueWidget()
+				]
+			]
+
+			+SVerticalBox::Slot()
+			.AutoHeight()
+			[
+				SNew(SHorizontalBox)
+
+				+SHorizontalBox::Slot()
+				.Padding(5, 2)
+				.FillWidth(0.3f)
+				[
+					PCACoeffNumPropertyHandle->CreatePropertyNameWidget()
+				]
+
+				+SHorizontalBox::Slot()
+				.Padding(5, 2)
+				.FillWidth(0.7f)
+				[
+					PCACoeffNumPropertyHandle->CreatePropertyValueWidget()
+				]
+
+			]
+
+			+SVerticalBox::Slot()
+			.AutoHeight()
+			[
+				SNew(SHorizontalBox)
+
+				+SHorizontalBox::Slot()
+				.Padding(5, 2)
+				.FillWidth(0.3f)
+				.VAlign(VAlign_Center)
+				[
+					SNew(STextBlock)
+					.Text(FText::FromString(TEXT("Submesh")))
+					.Font(IDetailLayoutBuilder::GetDetailFont())
+				]
+				+SHorizontalBox::Slot()
+				.Padding(5, 2)
+				.FillWidth(0.7f)
+				[
+					SNew(STextComboBox)
+					.OptionsSource(&SubMeshNames)
+					.OnSelectionChanged(SubMeshComboOnSelectionChangedDelegate)
+					.InitiallySelectedItem(SubMeshNames[InitMeshIndex])
+				]
+			]
+			+SVerticalBox::Slot()
+			.AutoHeight()
+			[
+				SNew(SHorizontalBox)
+
+				+SHorizontalBox::Slot()
+				.Padding(5, 2)
+				.FillWidth(0.3f)		
+				[
+					VertexMapPathHandle->CreatePropertyNameWidget()
+				]
+
+				+SHorizontalBox::Slot()
+				.Padding(5, 2)
+				.FillWidth(0.7f)
+				[
+					VertexMapPathHandle->CreatePropertyValueWidget()
+				]
+			]
+		];
+	}
+
 	void FNearestNeighborModelDetails::CustomizeDetails(IDetailLayoutBuilder& DetailBuilder)
 	{
 		// Create all the detail categories and add the properties of the base class.
@@ -84,31 +191,14 @@ namespace UE::NearestNeighborModel
 		Group->AddPropertyRow(DetailBuilder.GetProperty(UNearestNeighborModel::GetRecomputeDeltasPropertyName()));
 		Group->AddPropertyRow(DetailBuilder.GetProperty(UNearestNeighborModel::GetRecomputePCAPropertyName()));
 
-		// Cloth part settings
-		ClothPartCategoryBuilder->AddProperty(UNearestNeighborModel::GetClothPartEditorDataPropertyName());
-
-		FString VertexCountStr = "VertexCounts:";
-		if (NearestNeighborModel)
+		BuildSubMeshNames();
+		TSharedRef<IPropertyHandle> ClothPartDataPropertyHandle = DetailBuilder.GetProperty(UNearestNeighborModel::GetClothPartEditorDataPropertyName());
+		if (ClothPartDataPropertyHandle->AsArray().IsValid())
 		{
-			const int32 NumParts = NearestNeighborModel->GetNumParts();
-			VertexCountStr += "[";
-			for (int32 PartId = 0; PartId < NumParts; PartId++)
-			{
-				VertexCountStr += FString::Printf(TEXT("%d"), NearestNeighborModel->GetPartNumVerts(PartId));
-				if (PartId != NumParts - 1)
-				{
-					VertexCountStr += ",";
-				}
-			}
-			VertexCountStr += "]";
+			TSharedRef<FDetailArrayBuilder> PropertyBuilder = MakeShared<FDetailArrayBuilder>(ClothPartDataPropertyHandle, true, false, true);
+			PropertyBuilder->OnGenerateArrayElementWidget(FOnGenerateArrayElementWidget::CreateSP(this, &FNearestNeighborModelDetails::GenerateClothPartElementWidget));
+			ClothPartCategoryBuilder->AddCustomBuilder(PropertyBuilder);
 		}
-		FText VertexCountText = FText::FromString(VertexCountStr);
-		ClothPartCategoryBuilder->AddCustomRow(FText::FromString(""))
-			.WholeRowContent()
-			[
-				SNew(STextBlock)
-				.Text(VertexCountText)
-			];
 
 		// Nearest Neighbor settings
 		NearestNeighborCategoryBuilder->AddProperty(GET_MEMBER_NAME_STRING_CHECKED(UNearestNeighborModel, DecayFactor));
@@ -183,6 +273,30 @@ namespace UE::NearestNeighborModel
 		}
 	}
 
+	void FNearestNeighborModelDetails::BuildSubMeshNames()
+	{
+		SubMeshNames.Reset();
+		SubMeshNameMap.Reset();
+		if (NearestNeighborModel && NearestNeighborModel->GetSkeletalMesh() && NearestNeighborModel->GetSkeletalMesh()->GetImportedModel())
+		{
+			const FSkeletalMeshLODModel& LODModel =  NearestNeighborModel->GetSkeletalMesh()->GetImportedModel()->LODModels[0];
+			const TArray<FSkelMeshImportedMeshInfo>& SkelMeshInfos = LODModel.ImportedMeshInfos;
+			for (int32 i = 0; i < SkelMeshInfos.Num(); i++)
+			{
+				TSharedPtr<FString> StrPtr = MakeShared<FString>(SkelMeshInfos[i].Name.ToString());
+				SubMeshNames.Add(StrPtr);
+				SubMeshNameMap.Add(StrPtr, i);
+			}
+		}
+	}
+
+	void FNearestNeighborModelDetails::SubMeshComboSelectionChanged(TSharedPtr<FString> InSelectedItem, ESelectInfo::Type SelectInfo, int32 ArrayIndex)
+	{
+		if (NearestNeighborModel)
+		{
+			NearestNeighborModel->SetPartMeshIndex(ArrayIndex, SubMeshNameMap[InSelectedItem]);
+		}
+	}
 }	// namespace UE::NearestNeighborModel
 
 #undef LOCTEXT_NAMESPACE
diff --git a/Engine/Plugins/Animation/MLDeformer/NearestNeighborModel/Source/NearestNeighborModelEditor/Private/NearestNeighborTrainingModel.cpp b/Engine/Plugins/Animation/MLDeformer/NearestNeighborModel/Source/NearestNeighborModelEditor/Private/NearestNeighborTrainingModel.cpp
--- a/Engine/Plugins/Animation/MLDeformer/NearestNeighborModel/Source/NearestNeighborModelEditor/Private/NearestNeighborTrainingModel.cpp
+++ b/Engine/Plugins/Animation/MLDeformer/NearestNeighborModel/Source/NearestNeighborModelEditor/Private/NearestNeighborTrainingModel.cpp
@@ -36,7 +36,7 @@ const TArray<int32> UNearestNeighborTrainingModel::GetPartVertexMap(const int32
 int32 UNearestNeighborTrainingModel::SamplePart(int32 PartId, int32 Index)
 {
 	FNearestNeighborGeomCacheSampler* Sampler = static_cast<FNearestNeighborGeomCacheSampler*>(EditorModel->GetSampler());
-	int32 Result = Sampler->SamplePart(Index, NearestNeighborModel->PartVertexMap(PartId));
+	int32 Result = Sampler->SamplePart(Index, PartId);
 	PartSampleDeltas = Sampler->GetPartVertexDeltas();
 	SampleBoneRotations = Sampler->GetBoneRotations();
 	return Result;
diff --git a/Engine/Plugins/Animation/MLDeformer/NearestNeighborModel/Source/NearestNeighborModelEditor/Public/NearestNeighborGeomCacheSampler.h b/Engine/Plugins/Animation/MLDeformer/NearestNeighborModel/Source/NearestNeighborModelEditor/Public/NearestNeighborGeomCacheSampler.h
--- a/Engine/Plugins/Animation/MLDeformer/NearestNeighborModel/Source/NearestNeighborModelEditor/Public/NearestNeighborGeomCacheSampler.h
+++ b/Engine/Plugins/Animation/MLDeformer/NearestNeighborModel/Source/NearestNeighborModelEditor/Public/NearestNeighborGeomCacheSampler.h
@@ -10,13 +10,16 @@ namespace UE::NearestNeighborModel
 		: public UE::MLDeformer::FMLDeformerGeomCacheSampler
 	{
 	public:
-		virtual uint8 SamplePart(int32 InAnimFrameIndex, const TArray<uint32>& VertexMap);
+		virtual uint8 SamplePart(int32 InAnimFrameIndex, int32 PartId);
 		virtual void SampleKMeansAnim(const int32 SkeletonId);
 		virtual void SampleKMeansFrame(const int32 Frame);
 		const TArray<float>& GetPartVertexDeltas() const { return PartVertexDeltas; }
 		uint8 GeneratePartMeshMappings(const TArray<uint32>& VertexMap, bool bUsePartOnlyMesh);
+		uint8 GenerateMeshMappingIndices();
+
 		bool IsMeshMappingsEmpty() const { return MeshMappings.IsEmpty(); }
 	protected:
 		TArray<float> PartVertexDeltas;
+		TArray<int32> MeshMappingIndices;
 	};
 }
diff --git a/Engine/Plugins/Animation/MLDeformer/NearestNeighborModel/Source/NearestNeighborModelEditor/Public/NearestNeighborModelDetails.h b/Engine/Plugins/Animation/MLDeformer/NearestNeighborModel/Source/NearestNeighborModelEditor/Public/NearestNeighborModelDetails.h
--- a/Engine/Plugins/Animation/MLDeformer/NearestNeighborModel/Source/NearestNeighborModelEditor/Public/NearestNeighborModelDetails.h
+++ b/Engine/Plugins/Animation/MLDeformer/NearestNeighborModel/Source/NearestNeighborModelEditor/Public/NearestNeighborModelDetails.h
@@ -6,6 +6,8 @@
 #include "MLDeformerMorphModelDetails.h"
 
 class UNearestNeighborModel;
+class IPropertyHandle;
+class IDetailChildrenBuilder;
 
 namespace UE::NearestNeighborModel
 {
@@ -27,6 +29,8 @@ namespace UE::NearestNeighborModel
 		virtual bool UpdateMemberPointers(const TArray<TWeakObjectPtr<UObject>>& Objects) override;
 		virtual void AddUpdateResultText(IDetailCategoryBuilder* CategoryBuilder, uint8 Result);
 
+		void GenerateClothPartElementWidget(TSharedRef<IPropertyHandle> PropertyHandle, int32 ArrayIndex, IDetailChildrenBuilder& ChildrenBuilder);
+
 	protected:
 		UNearestNeighborModel* NearestNeighborModel = nullptr;
 		FNearestNeighborEditorModel* NearestNeighborEditorModel = nullptr;
@@ -35,5 +39,10 @@ namespace UE::NearestNeighborModel
 		IDetailCategoryBuilder* ClothPartCategoryBuilder = nullptr;
 		IDetailCategoryBuilder* NearestNeighborCategoryBuilder = nullptr;
 		IDetailCategoryBuilder* KMeansCategoryBuilder = nullptr;
+
+		TArray<TSharedPtr<FString> > SubMeshNames;
+		TMap<TSharedPtr<FString>, int32> SubMeshNameMap;
+		void BuildSubMeshNames();
+		void SubMeshComboSelectionChanged(TSharedPtr<FString> InSelectedItem, ESelectInfo::Type SelectInfo, int32 ArrayIndex);
 	};
 }	// namespace UE::NearestNeighborModel
