NearestNeighborEditor: add option to offset nearest neighbor actors

diff --git a/Engine/Plugins/Animation/MLDeformer/NearestNeighborModel/Source/NearestNeighborModel/Public/NearestNeighborModelVizSettings.h b/Engine/Plugins/Animation/MLDeformer/NearestNeighborModel/Source/NearestNeighborModel/Public/NearestNeighborModelVizSettings.h
--- a/Engine/Plugins/Animation/MLDeformer/NearestNeighborModel/Source/NearestNeighborModel/Public/NearestNeighborModelVizSettings.h
+++ b/Engine/Plugins/Animation/MLDeformer/NearestNeighborModel/Source/NearestNeighborModel/Public/NearestNeighborModelVizSettings.h
@@ -16,4 +16,14 @@ class NEARESTNEIGHBORMODEL_API UNearestNeighborModelVizSettings
 	: public UMLDeformerMorphModelVizSettings
 {
 	GENERATED_BODY()
+public:
+#if	WITH_EDITORONLY_DATA
+	void SetNearestNeighborActorsOffset(float InOffset) { NearestNeighborActorsOffset = InOffset; }
+	float GetNearestNeighborActorsOffset() const { return NearestNeighborActorsOffset; }
+	static FName GetNearestNeighborActorsOffsetPropertyName() { return GET_MEMBER_NAME_CHECKED(UNearestNeighborModelVizSettings, NearestNeighborActorsOffset); }
+
+protected:
+	UPROPERTY(EditAnywhere, Category = "Live Settings")
+	float NearestNeighborActorsOffset = 2.0f;
+#endif
 };
diff --git a/Engine/Plugins/Animation/MLDeformer/NearestNeighborModel/Source/NearestNeighborModelEditor/Private/NearestNeighborEditorModel.cpp b/Engine/Plugins/Animation/MLDeformer/NearestNeighborModel/Source/NearestNeighborModelEditor/Private/NearestNeighborEditorModel.cpp
--- a/Engine/Plugins/Animation/MLDeformer/NearestNeighborModel/Source/NearestNeighborModelEditor/Private/NearestNeighborEditorModel.cpp
+++ b/Engine/Plugins/Animation/MLDeformer/NearestNeighborModel/Source/NearestNeighborModelEditor/Private/NearestNeighborEditorModel.cpp
@@ -110,6 +110,11 @@ namespace UE::NearestNeighborModel
 		}
 	}
 
+	UNearestNeighborModelVizSettings* FNearestNeighborEditorModel::GetNearestNeighborModelVizSettings() const
+	{
+		return Cast<UNearestNeighborModelVizSettings>(GetMorphModel()->GetVizSettings()); 
+	}
+
 	void FNearestNeighborEditorModel::CreateActors(const TSharedRef<IPersonaPreviewScene>& InPersonaPreviewScene)
 	{
 		FMLDeformerMorphModelEditorModel::CreateActors(InPersonaPreviewScene);
@@ -122,11 +127,17 @@ namespace UE::NearestNeighborModel
 	void FNearestNeighborEditorModel::Tick(FEditorViewportClient* ViewportClient, float DeltaTime)
 	{
 		FMLDeformerEditorModel::Tick(ViewportClient, DeltaTime);
-		for (FNearestNeighborEditorModelActor* NearestNeighborActor : NearestNeighborActors)
+		if (!NearestNeighborActors.IsEmpty())
 		{
-			if (NearestNeighborActor)
+			const UNearestNeighborModelVizSettings* NNViz = GetNearestNeighborModelVizSettings();
+			const float Offset = NNViz->GetNearestNeighborActorsOffset();
+			for (FNearestNeighborEditorModelActor* NearestNeighborActor : NearestNeighborActors)
 			{
-				NearestNeighborActor->TickNearestNeighborActor();
+				if (NearestNeighborActor)
+				{
+					NearestNeighborActor->TickNearestNeighborActor();
+					NearestNeighborActor->SetMeshOffsetFactor(Offset);
+				}
 			}
 		}
 	}
@@ -141,6 +152,9 @@ namespace UE::NearestNeighborModel
 			NearestNeighborActors.Reset();
 		}
 		NearestNeighborActors.SetNumZeroed(NumParts);
+
+		const UNearestNeighborModelVizSettings* NNViz = GetNearestNeighborModelVizSettings();
+		const float Offset = NNViz->GetNearestNeighborActorsOffset();
 		
 		for(int32 PartId = StartIndex; PartId < NumParts; PartId++)
 		{
@@ -159,7 +173,7 @@ namespace UE::NearestNeighborModel
 			FNearestNeighborEditorModelActor* NearestNeighborActor = static_cast<FNearestNeighborEditorModelActor*>(EditorActors.Last());
 			UMLDeformerComponent* MLDeformerComponent = GetTestMLDeformerComponent();
 			NearestNeighborActor->InitNearestNeighborActor(PartId, MLDeformerComponent);
-			NearestNeighborActor->SetMeshOffsetFactor(2.0f);
+			NearestNeighborActor->SetMeshOffsetFactor(Offset);
 			NearestNeighborActors[PartId] = NearestNeighborActor;
 		}
 	}
diff --git a/Engine/Plugins/Animation/MLDeformer/NearestNeighborModel/Source/NearestNeighborModelEditor/Private/NearestNeighborModelVizSettingsDetails.cpp b/Engine/Plugins/Animation/MLDeformer/NearestNeighborModel/Source/NearestNeighborModelEditor/Private/NearestNeighborModelVizSettingsDetails.cpp
new file mode 100644
--- /dev/null
+++ b/Engine/Plugins/Animation/MLDeformer/NearestNeighborModel/Source/NearestNeighborModelEditor/Private/NearestNeighborModelVizSettingsDetails.cpp
@@ -0,0 +1,24 @@
+// Copyright Epic Games, Inc. All Rights Reserved.
+
+#include "NearestNeighborModelVizSettingsDetails.h"
+
+#include "DetailLayoutBuilder.h"
+#include "DetailCategoryBuilder.h"
+#include "DetailWidgetRow.h"
+#include "IDetailsView.h"
+#include "IDetailGroup.h"
+#include "NearestNeighborModelVizSettings.h"
+
+#define LOCTEXT_NAMESPACE "MLDeformerMorphModelVizSettingsDetails"
+
+namespace UE::NearestNeighborModel
+{
+	void FNearestNeighborModelVizSettingsDetails::AddAdditionalSettings()
+	{
+		FMLDeformerMorphModelVizSettingsDetails::AddAdditionalSettings();
+		IDetailGroup& NNGroup = LiveSettingsCategory->AddGroup("Nearest Neighbor", LOCTEXT("NearestNeighborLabel", "Nearest Neighbor"), false, true);
+		NNGroup.AddPropertyRow(DetailLayoutBuilder->GetProperty(UNearestNeighborModelVizSettings::GetNearestNeighborActorsOffsetPropertyName(), UNearestNeighborModelVizSettings::StaticClass()));
+	}
+
+};
+#undef LOCTEXT_NAMESPACE
diff --git a/Engine/Plugins/Animation/MLDeformer/NearestNeighborModel/Source/NearestNeighborModelEditor/Public/NearestNeighborEditorModel.h b/Engine/Plugins/Animation/MLDeformer/NearestNeighborModel/Source/NearestNeighborModelEditor/Public/NearestNeighborEditorModel.h
--- a/Engine/Plugins/Animation/MLDeformer/NearestNeighborModel/Source/NearestNeighborModelEditor/Public/NearestNeighborEditorModel.h
+++ b/Engine/Plugins/Animation/MLDeformer/NearestNeighborModel/Source/NearestNeighborModelEditor/Public/NearestNeighborEditorModel.h
@@ -48,7 +48,8 @@ namespace UE::NearestNeighborModel
 
 		// Some helpers that cast to this model's variants of some classes.
 		UNearestNeighborModel* GetNearestNeighborModel() const { return static_cast<UNearestNeighborModel*>(Model); }
-	
+		UNearestNeighborModelVizSettings* GetNearestNeighborModelVizSettings() const;
+
 		// Recomputes nearest neighbor coeffcients and nearest neighbor vertex offsets. 
 		virtual uint8 UpdateNearestNeighborData();
 
diff --git a/Engine/Plugins/Animation/MLDeformer/NearestNeighborModel/Source/NearestNeighborModelEditor/Public/NearestNeighborModelVizSettingsDetails.h b/Engine/Plugins/Animation/MLDeformer/NearestNeighborModel/Source/NearestNeighborModelEditor/Public/NearestNeighborModelVizSettingsDetails.h
--- a/Engine/Plugins/Animation/MLDeformer/NearestNeighborModel/Source/NearestNeighborModelEditor/Public/NearestNeighborModelVizSettingsDetails.h
+++ b/Engine/Plugins/Animation/MLDeformer/NearestNeighborModel/Source/NearestNeighborModelEditor/Public/NearestNeighborModelVizSettingsDetails.h
@@ -16,5 +16,7 @@ namespace UE::NearestNeighborModel
 		{
 			return MakeShareable(new FNearestNeighborModelVizSettingsDetails());
 		}
+		
+		virtual void AddAdditionalSettings() override;
 	};
 }	// namespace UE::NearestNeighborModel
