Nearestneighboreditor: detect if part-only mesh is used and add error message if UsePartOnlyMesh is not set correctly.

diff --git a/Engine/Plugins/Animation/MLDeformer/NearestNeighborModel/Source/NearestNeighborModel/Public/NearestNeighborModel.h b/Engine/Plugins/Animation/MLDeformer/NearestNeighborModel/Source/NearestNeighborModel/Public/NearestNeighborModel.h
--- a/Engine/Plugins/Animation/MLDeformer/NearestNeighborModel/Source/NearestNeighborModel/Public/NearestNeighborModel.h
+++ b/Engine/Plugins/Animation/MLDeformer/NearestNeighborModel/Source/NearestNeighborModel/Public/NearestNeighborModel.h
@@ -127,6 +127,11 @@ namespace UE::NearestNeighborModel
 		ERROR = 1,
 		WARNING = 2
 	};
+
+	static bool HasError(uint8 Flag)
+	{
+		return (Flag & EUpdateResult::ERROR) != 0;
+	}
 };
 
 /**
diff --git a/Engine/Plugins/Animation/MLDeformer/NearestNeighborModel/Source/NearestNeighborModelEditor/Private/NearestNeighborEditorModel.cpp b/Engine/Plugins/Animation/MLDeformer/NearestNeighborModel/Source/NearestNeighborModelEditor/Private/NearestNeighborEditorModel.cpp
--- a/Engine/Plugins/Animation/MLDeformer/NearestNeighborModel/Source/NearestNeighborModelEditor/Private/NearestNeighborEditorModel.cpp
+++ b/Engine/Plugins/Animation/MLDeformer/NearestNeighborModel/Source/NearestNeighborModelEditor/Private/NearestNeighborEditorModel.cpp
@@ -190,11 +190,6 @@ namespace UE::NearestNeighborModel
 		NearestNeighborModel->InitInputInfo();
 	}
 
-	bool HasError(uint8 Flag)
-	{
-		return (Flag & EUpdateResult::ERROR) != 0;
-	}
-
 	ETrainingResult FNearestNeighborEditorModel::Train()
 	{
 		UNearestNeighborModel* NearestNeighborModel = GetNearestNeighborModel();
@@ -254,12 +249,6 @@ namespace UE::NearestNeighborModel
 						return ReturnCode;
 					}
 
-					if (GeomCacheSampler->IsMeshMappingsEmpty())
-					{
-						UE_LOG(LogNearestNeighborModel, Error, TEXT("Failed to generate mappings between skeletal mesh and geometry cache."));
-						return EUpdateResult::ERROR;
-					}
-
 					NumTrainingFramesOverride = NumFrames;
 
 					if (NumNeighborsFromGeomCache == NumNeighborsFromAnimSequence)
diff --git a/Engine/Plugins/Animation/MLDeformer/NearestNeighborModel/Source/NearestNeighborModelEditor/Private/NearestNeighborGeomCacheSampler.cpp b/Engine/Plugins/Animation/MLDeformer/NearestNeighborModel/Source/NearestNeighborModelEditor/Private/NearestNeighborGeomCacheSampler.cpp
--- a/Engine/Plugins/Animation/MLDeformer/NearestNeighborModel/Source/NearestNeighborModelEditor/Private/NearestNeighborGeomCacheSampler.cpp
+++ b/Engine/Plugins/Animation/MLDeformer/NearestNeighborModel/Source/NearestNeighborModelEditor/Private/NearestNeighborGeomCacheSampler.cpp
@@ -30,11 +30,11 @@ namespace UE::NearestNeighborModel
 			const FSkeletalMeshLODModel& LODModel = ImportedModel->LODModels[LODIndex];
 			const TArray<FSkelMeshImportedMeshInfo>& SkelMeshInfos = LODModel.ImportedMeshInfos;
 
-			if (MeshMappings.IsEmpty())
+			if (HasError(CheckMeshMappingsEmpty()))
 			{
-				UE_LOG(LogNearestNeighborModel, Error, TEXT("SamplePart: MeshMappings is empty."));
 				return EUpdateResult::ERROR;
 			}
+
 			if (PartId >= MeshMappingIndices.Num())
 			{
 				UE_LOG(LogNearestNeighborModel, Error, TEXT("SamplePart: MeshMappingIndices.Num()=%d is smaller than PartId %d"), MeshMappingIndices.Num(), PartId);
@@ -129,11 +129,8 @@ namespace UE::NearestNeighborModel
 				}
 				GeomCacheMeshDatas.Reset();
 				GeomCacheMeshDatas.AddDefaulted(MeshMappings.Num());
-				if(MeshMappings.IsEmpty())
-				{
-					UE_LOG(LogNearestNeighborModel, Error, TEXT("Unable to match skeletal mesh with geometry cache."));
-					Result = EUpdateResult::ERROR;
-				}
+
+				Result |= CheckMeshMappingsEmpty();
 				return Result;
 			}
 			FSkeletalMeshModel* ImportedModel = SkeletalMesh->GetImportedModel();
@@ -178,6 +175,14 @@ namespace UE::NearestNeighborModel
 							continue;
 						}
 
+						const int32 NumVertsFromGeomCache = FMath::Max(GeomCacheMeshData.ImportedVertexNumbers) + 1;
+						const int32 NumVertsFromVertexMap = VertexMap.Num();
+						Result |= CheckGeomCacheVertCount(NumVertsFromGeomCache, NumVertsFromVertexMap);
+						if (HasError(Result))
+						{
+							return Result;
+						}
+
 						// Create a new mesh mapping entry.
 						MeshMappings.AddDefaulted();
 						UE::MLDeformer::FMLDeformerGeomCacheMeshMapping& Mapping = MeshMappings.Last();
@@ -215,10 +220,47 @@ namespace UE::NearestNeighborModel
 			Result |= EUpdateResult::WARNING;
 			UE_LOG(LogNearestNeighborModel, Warning, TEXT("SkeletalMesh or GeometryCache is none. No mapping is generated"));
 		}
+		Result |= CheckMeshMappingsEmpty();
 		Result |= GenerateMeshMappingIndices();
 		return Result;
 	}
 
+	uint8 FNearestNeighborGeomCacheSampler::CheckGeomCacheVertCount(int32 NumVertsFromGeomCache, int32 NumVertsFromVertexMap) const
+	{
+		uint8 ResultCode = EUpdateResult::SUCCESS;		
+		UNearestNeighborModel* NearestNeighborModel = static_cast<UNearestNeighborModel*>(Model);
+		if (NearestNeighborModel->GetUsePartOnlyMesh())
+		{
+			if (NumVertsFromGeomCache != NumVertsFromVertexMap)
+			{
+				UE_LOG(LogNearestNeighborModel, Error, TEXT("Vertex count mismatch: GeomCache has %d vertices but vertex map has %d vertices. Maybe turn off UsePartOnlyMesh."), NumVertsFromGeomCache, NumVertsFromVertexMap);
+				return EUpdateResult::ERROR;
+			}
+		}
+		else
+		{
+			if (NumVertsFromVertexMap > NumVertsFromGeomCache)
+			{
+				UE_LOG(LogNearestNeighborModel, Error, TEXT("Vertex count in vertex map %d is larger than the vertex count in geometry cache %d. Something is wrong with vertex map or geometry cache"), NumVertsFromVertexMap, NumVertsFromGeomCache);
+				return EUpdateResult::ERROR;
+			}
+		}
+		return ResultCode;
+	}
+
+	uint8 FNearestNeighborGeomCacheSampler::CheckMeshMappingsEmpty() const
+	{
+		if(MeshMappings.IsEmpty())
+		{
+			UE_LOG(LogNearestNeighborModel, Error, TEXT("MeshMappings is empty. Unable to match skeletal mesh with geometry cache."));
+			return EUpdateResult::ERROR;
+		}
+		else
+		{
+			return EUpdateResult::SUCCESS;
+		}
+	}
+
 	void FNearestNeighborGeomCacheSampler::SampleKMeansAnim(const int32 SkeletonId)
 	{
 		UNearestNeighborModel* NearestNeighborModel = static_cast<UNearestNeighborModel*>(Model);
diff --git a/Engine/Plugins/Animation/MLDeformer/NearestNeighborModel/Source/NearestNeighborModelEditor/Public/NearestNeighborGeomCacheSampler.h b/Engine/Plugins/Animation/MLDeformer/NearestNeighborModel/Source/NearestNeighborModelEditor/Public/NearestNeighborGeomCacheSampler.h
--- a/Engine/Plugins/Animation/MLDeformer/NearestNeighborModel/Source/NearestNeighborModelEditor/Public/NearestNeighborGeomCacheSampler.h
+++ b/Engine/Plugins/Animation/MLDeformer/NearestNeighborModel/Source/NearestNeighborModelEditor/Public/NearestNeighborGeomCacheSampler.h
@@ -16,8 +16,9 @@ namespace UE::NearestNeighborModel
 		const TArray<float>& GetPartVertexDeltas() const { return PartVertexDeltas; }
 		uint8 GeneratePartMeshMappings(const TArray<uint32>& VertexMap, bool bUsePartOnlyMesh);
 		uint8 GenerateMeshMappingIndices();
+		uint8 CheckGeomCacheVertCount(int32 NumVertsFromGeomCache, int32 NumVertsFromVertexMap) const;
+		uint8 CheckMeshMappingsEmpty() const;
 
-		bool IsMeshMappingsEmpty() const { return MeshMappings.IsEmpty(); }
 	protected:
 		TArray<float> PartVertexDeltas;
 		TArray<int32> MeshMappingIndices;
