[MLDeformer] Make the nearest neighbor model support the masking and morph target normals.

#trivial
#rnx

diff --git a/Engine/Plugins/Animation/MLDeformer/NearestNeighborModel/Source/NearestNeighborModelEditor/Private/NearestNeighborEditorModel.cpp b/Engine/Plugins/Animation/MLDeformer/NearestNeighborModel/Source/NearestNeighborModelEditor/Private/NearestNeighborEditorModel.cpp
--- a/Engine/Plugins/Animation/MLDeformer/NearestNeighborModel/Source/NearestNeighborModelEditor/Private/NearestNeighborEditorModel.cpp
+++ b/Engine/Plugins/Animation/MLDeformer/NearestNeighborModel/Source/NearestNeighborModelEditor/Private/NearestNeighborEditorModel.cpp
@@ -557,7 +557,16 @@ namespace UE::NearestNeighborModel
 
 		const int32 LOD = 0;
 		TArray<UMorphTarget*> MorphTargets;
-		CreateEngineMorphTargets(MorphTargets, Deltas, FString("NNMorphTarget_"), LOD, NearestNeighborModel->GetMorphTargetDeltaThreshold());
+		CreateEngineMorphTargets(
+			MorphTargets, 
+			Deltas, 
+			FString("NNMorphTarget_"),
+			LOD,
+			NearestNeighborModel->GetMorphTargetDeltaThreshold(),
+			NearestNeighborModel->GetIncludeMorphTargetNormals(),
+			NearestNeighborModel->GetMaskChannel(),
+			NearestNeighborModel->GetInvertMaskChannel());
+
 		check(NearestNeighborModel->GetMorphTargetSet().IsValid());
 		FMorphTargetVertexInfoBuffers& MorphBuffers = NearestNeighborModel->GetMorphTargetSet()->MorphBuffers;
 		CompressEngineMorphTargets(MorphBuffers, MorphTargets, LOD, NearestNeighborModel->GetMorphTargetErrorTolerance());
