Add input modifier

diff --git a/Engine/Plugins/Animation/MLDeformer/NearestNeighborModel/Source/NearestNeighborModel/Private/NearestNeighborModel.cpp b/Engine/Plugins/Animation/MLDeformer/NearestNeighborModel/Source/NearestNeighborModel/Private/NearestNeighborModel.cpp
--- a/Engine/Plugins/Animation/MLDeformer/NearestNeighborModel/Source/NearestNeighborModel/Private/NearestNeighborModel.cpp
+++ b/Engine/Plugins/Animation/MLDeformer/NearestNeighborModel/Source/NearestNeighborModel/Private/NearestNeighborModel.cpp
@@ -100,22 +100,38 @@ TArray<uint32> AddConstant(const TArray<uint32> &InArr, uint32 Constant)
 	return OutArr;
 }
 
-void UNearestNeighborModel::ClipInputs(float* InputPtr, int NumInputs)
+TArray<float> UNearestNeighborModel::ClipInputs(const TArray<float>& Inputs) const
 {
-	if(InputsMin.Num() == NumInputs && InputsMax.Num() == NumInputs)
+	TArray<float> Result = Inputs;
+	ClipInputs(Result.GetData(), Result.Num());
+	return Result;
+}
+
+void UNearestNeighborModel::ClipInputs(float* InputPtr, int NumInputs) const
+{
+	if(InputsMin.Num() >= NumInputs && InputsMax.Num() >= NumInputs)
 	{
 		for(int32 i = 0; i < NumInputs; i++)
 		{
-			if (InputPtr[i] > InputsMax[i])
+			float Max = InputsMax[i];
+			float Min = InputsMin[i];
+			if (bUseInputMultipliers && i / 3 < InputMultipliers.Num())
 			{
-				InputPtr[i] = InputsMax[i];
+				const float Multiplier = InputMultipliers[i / 3][i % 3];
+				Max *= Multiplier;
+				Min *= Multiplier;	
 			}
-			if (InputPtr[i] < InputsMin[i])
+
+			if (InputPtr[i] > Max)
 			{
-				InputPtr[i] = InputsMin[i];
+				InputPtr[i] = Max;
+			}
+			if (InputPtr[i] < Min)
+			{
+				InputPtr[i] = Min;
 			}
 		}
-	}	
+	}
 }
 
 int32 UNearestNeighborModel::GetTotalNumPCACoeffs() const
@@ -179,6 +195,20 @@ void UNearestNeighborModel::UpdateNetworkOutputDim()
 	}
 }
 
+void UNearestNeighborModel::UpdateInputMultipliers()
+{
+	if (bUseInputMultipliers)
+	{
+		const int32 NumMultipliers = InputMultipliers.Num();
+		const int32 NumBones = GetBoneIncludeList().Num();
+		InputMultipliers.SetNum(NumBones);
+		for (int32 Index = NumMultipliers; Index < NumBones; Index++)
+		{
+			InputMultipliers[Index] = FVector3f(1.f, 1.f, 1.f);
+		}
+	}
+}
+
 UE::NearestNeighborModel::EUpdateResult UNearestNeighborModel::UpdateVertexMap(int32 PartId, const FString& VertexMapPath, const FSkelMeshImportedMeshInfo& Info)
 {
 	using namespace UE::NearestNeighborModel;
diff --git a/Engine/Plugins/Animation/MLDeformer/NearestNeighborModel/Source/NearestNeighborModel/Public/NearestNeighborModel.h b/Engine/Plugins/Animation/MLDeformer/NearestNeighborModel/Source/NearestNeighborModel/Public/NearestNeighborModel.h
--- a/Engine/Plugins/Animation/MLDeformer/NearestNeighborModel/Source/NearestNeighborModel/Public/NearestNeighborModel.h
+++ b/Engine/Plugins/Animation/MLDeformer/NearestNeighborModel/Source/NearestNeighborModel/Public/NearestNeighborModel.h
@@ -221,7 +221,10 @@ public:
 	UFUNCTION(BlueprintCallable, Category = "Nearest Neighbor Model")
 	void SetNeighborOffsets(int32 PartId, const TArray<float>& NeighborOffsets) { ClothPartData[PartId].NeighborOffsets = NeighborOffsets; }
 
-	void ClipInputs(float* InputPtr, int NumInputs);
+	UFUNCTION(BlueprintPure, Category = "Nearest Neighbor Model")
+	TArray<float> ClipInputs(const TArray<float>& Input) const;
+
+	void ClipInputs(float* InputPtr, int NumInputs) const;
 
 	bool CheckPCAData(int32 PartId) const;
 
@@ -244,6 +247,7 @@ public:
 	void UpdatePCACoeffNums();
 	void UpdateNetworkSize();
 	void UpdateMorphTargetSize();
+	void UpdateInputMultipliers();
 
 	void InvalidateClothPartData() { bClothPartDataValid = false; bNearestNeighborDataValid = false; bMorphTargetDataValid = false; }
 	void ValidateClothPartData() { bClothPartDataValid = true; }
@@ -370,13 +374,19 @@ protected:
 #endif
 
 public:
+	UPROPERTY(EditAnywhere, BlueprintReadWrite, Category = "Debugging")
+	bool bUseInputMultipliers = false;
+
+	UPROPERTY(EditAnywhere, BlueprintReadWrite, Category = "Debugging")
+	TArray<FVector3f> InputMultipliers;
+
 	UPROPERTY()
 	TArray<FClothPartData> ClothPartData;
 
-	UPROPERTY(BlueprintReadWrite, Category = "Network Inputs")
+	UPROPERTY(BlueprintReadWrite, Category = "Network IO")
 	TArray<float> InputsMin;
 
-	UPROPERTY(BlueprintReadWrite, Category = "Network Inputs")
+	UPROPERTY(BlueprintReadWrite, Category = "Network IO")
 	TArray<float> InputsMax;
 
 	UPROPERTY(EditAnywhere, BlueprintReadWrite, Category = "KMeans Pose Generator")
diff --git a/Engine/Plugins/Animation/MLDeformer/NearestNeighborModel/Source/NearestNeighborModelEditor/Private/NearestNeighborEditorModel.cpp b/Engine/Plugins/Animation/MLDeformer/NearestNeighborModel/Source/NearestNeighborModelEditor/Private/NearestNeighborEditorModel.cpp
--- a/Engine/Plugins/Animation/MLDeformer/NearestNeighborModel/Source/NearestNeighborModelEditor/Private/NearestNeighborEditorModel.cpp
+++ b/Engine/Plugins/Animation/MLDeformer/NearestNeighborModel/Source/NearestNeighborModelEditor/Private/NearestNeighborEditorModel.cpp
@@ -486,6 +486,8 @@ namespace UE::NearestNeighborModel
 		}
 		InitTestMLDeformerPreviousWeights();
 		GetNearestNeighborModel()->ValidateMorphTargetData();
+
+		GetNearestNeighborModel()->UpdateInputMultipliers();
 	}
 
 	uint8 FNearestNeighborEditorModel::InitMorphTargets()
