NearestNeighborEditor: add KMeans clustering

diff --git a/Engine/Plugins/Animation/MLDeformer/NearestNeighborModel/Source/NearestNeighborModel/Public/NearestNeighborModel.h b/Engine/Plugins/Animation/MLDeformer/NearestNeighborModel/Source/NearestNeighborModel/Public/NearestNeighborModel.h
--- a/Engine/Plugins/Animation/MLDeformer/NearestNeighborModel/Source/NearestNeighborModel/Public/NearestNeighborModel.h
+++ b/Engine/Plugins/Animation/MLDeformer/NearestNeighborModel/Source/NearestNeighborModel/Public/NearestNeighborModel.h
@@ -372,8 +372,19 @@ protected:
 	UPROPERTY(EditAnywhere, BlueprintReadWrite, Category = "File Cache")
 	bool bRecomputePCA = true;
 #endif
-
+	
 public:
+#if WITH_EDITORONLY_DATA
+	UPROPERTY(EditAnywhere, BlueprintReadWrite, Category = "KMeans Pose Generator")
+	TArray<TObjectPtr<UAnimSequence>> SourceAnims;
+
+	UPROPERTY(EditAnywhere, BlueprintReadWrite, Category = "KMeans Pose Generator", meta = (ClampMin = "1"))
+	int32 NumClusters;
+
+	UPROPERTY(EditAnywhere, BlueprintReadWrite, Category = "KMeans Pose Generator")
+	int32 KMeansPartId = 0;
+#endif
+
 	UPROPERTY(EditAnywhere, BlueprintReadWrite, Category = "Debugging")
 	bool bUseInputMultipliers = false;
 
@@ -389,12 +400,6 @@ public:
 	UPROPERTY(BlueprintReadWrite, Category = "Network IO")
 	TArray<float> InputsMax;
 
-	UPROPERTY(EditAnywhere, BlueprintReadWrite, Category = "KMeans Pose Generator")
-	TArray<TObjectPtr<UAnimSequence>> SourceSkeletons;
-
-	UPROPERTY(EditAnywhere, BlueprintReadWrite, Category = "KMeans Pose Generator", meta = (ClampMin = "1"))
-	int32 NumClusters;
-
 	UPROPERTY(EditAnywhere, BlueprintReadWrite, Category = "Nearest Neighbors", META = (ClampMin = "0", ClampMax = "1"))
 	float DecayFactor = 0.85f;
 
diff --git a/Engine/Plugins/Animation/MLDeformer/NearestNeighborModel/Source/NearestNeighborModelEditor/Private/NearestNeighborEditorModel.cpp b/Engine/Plugins/Animation/MLDeformer/NearestNeighborModel/Source/NearestNeighborModelEditor/Private/NearestNeighborEditorModel.cpp
--- a/Engine/Plugins/Animation/MLDeformer/NearestNeighborModel/Source/NearestNeighborModelEditor/Private/NearestNeighborEditorModel.cpp
+++ b/Engine/Plugins/Animation/MLDeformer/NearestNeighborModel/Source/NearestNeighborModelEditor/Private/NearestNeighborEditorModel.cpp
@@ -12,12 +12,17 @@
 #include "Components/ExternalMorphSet.h"
 #include "MLDeformerComponent.h"
 #include "MLDeformerEditorToolkit.h"
+#include "MLDeformerAsset.h"
 #include "NeuralNetwork.h"
 #include "GeometryCache.h"
 #include "GeometryCacheComponent.h"
 #include "Animation/DebugSkelMeshComponent.h"
 #include "Animation/MorphTarget.h"
 #include "Animation/AnimSequence.h"
+#include "PackageTools.h"
+#include "ObjectTools.h"
+#include "UObject/SavePackage.h"
+
 
 #define LOCTEXT_NAMESPACE "NearestNeighborEditorModel"
 
@@ -393,8 +398,63 @@ namespace UE::NearestNeighborModel
 		return EUpdateResult::ERROR;
 	}
 
+	// write a function to return an animation sequence and a boolean
+
+	template<typename T>
+	T* CreateObjectInstance(const FString& PackageName)
+	{
+		// Parent package to place new mesh
+		UPackage* Package = nullptr;
+		FString NewPackageName = PackageName;
+		const FString ObjectName = FPaths::GetBaseFilename(PackageName);
+
+		// Setup package name and create one accordingly
+		NewPackageName = UPackageTools::SanitizePackageName(NewPackageName);
+		Package = CreatePackage(*NewPackageName);
+
+		const FString SanitizedObjectName = ObjectTools::SanitizeObjectName(ObjectName);
+
+		T* ExistingTypedObject = FindObject<T>(Package, *SanitizedObjectName);
+		UObject* ExistingObject = FindObject<UObject>(Package, *SanitizedObjectName);
+
+		if (ExistingTypedObject != nullptr)
+		{
+			ExistingTypedObject->PreEditChange(nullptr);
+		}
+		else if (ExistingObject != nullptr)
+		{
+			// Replacing an object.  Here we go!
+			// Delete the existing object
+			const bool bDeleteSucceeded = ObjectTools::DeleteSingleObject(ExistingObject);
+
+			if (bDeleteSucceeded)
+			{
+				// Force GC so we can cleanly create a new asset (and not do an 'in place' replacement)
+				CollectGarbage(GARBAGE_COLLECTION_KEEPFLAGS);
+
+				// Create a package for each mesh
+				Package = CreatePackage(*NewPackageName);
+			}
+			else
+			{
+				// failed to delete
+				UE_LOG(LogNearestNeighborModel, Error, TEXT("Failed to delete existing object %s"), *SanitizedObjectName);
+				return nullptr;
+			}
+		}
+
+		if (Package == nullptr)
+		{
+			UE_LOG(LogNearestNeighborModel, Error, TEXT("Failed to create package %s"), *NewPackageName);
+			return nullptr;
+		}
+
+		return NewObject<T>(Package, FName(*SanitizedObjectName),  RF_Public | RF_Standalone);
+	}
+
 	void FNearestNeighborEditorModel::KMeansClusterPoses()
 	{
+		KMeansClusterResult = EUpdateResult::SUCCESS;
 		UNearestNeighborModel *NearestNeighborModel = static_cast<UNearestNeighborModel*>(Model);
 		check(NearestNeighborModel != nullptr);
 
@@ -402,22 +462,149 @@ namespace UE::NearestNeighborModel
 		if(NeuralNetwork && NeuralNetwork->IsLoaded())
 		{
 			const FString SavePath = GetTrainedNetworkOnnxFile();
-			if (!FPaths::FileExists(SavePath) || !NearestNeighborModel->GetUseFileCache())
+			if (FPaths::DirectoryExists(FPaths::GetPath(SavePath)))
 			{
 				UE_LOG(LogNearestNeighborModel, Display, TEXT("Saving to %s"), *SavePath);
 				NeuralNetwork->Save(SavePath);
 			}
+			else
+			{
+				UE_LOG(LogNearestNeighborModel, Error, TEXT("Path %s does not exist."), *FPaths::GetPath(SavePath));
+				KMeansClusterResult |= EUpdateResult::ERROR;
+				return;
+			}
 		}
 		else
 		{
 			UE_LOG(LogNearestNeighborModel, Warning, TEXT("Network is not available. Nothing will be done."));
+			KMeansClusterResult |= EUpdateResult::WARNING;
 			return;
 		}
+		if (NearestNeighborModel->KMeansPartId >= NearestNeighborModel->GetNumParts())
+		{
+			UE_LOG(LogNearestNeighborModel, Error, TEXT("KMeansPartId %d is out of range [0, %d). Nothing will be done."), NearestNeighborModel->KMeansPartId, NearestNeighborModel->GetNumParts());
+			KMeansClusterResult |= EUpdateResult::ERROR;
+			return;
+		}
+		if (NearestNeighborModel->SourceAnims.Num() == 0)
+		{
+			UE_LOG(LogNearestNeighborModel, Error, TEXT("No source anims found."));
+			KMeansClusterResult |= EUpdateResult::ERROR;
+			return;
+		}
+		for (int32 i = 0; i < NearestNeighborModel->SourceAnims.Num(); i++)
+		{
+			if (NearestNeighborModel->SourceAnims[i] == nullptr)
+			{
+				UE_LOG(LogNearestNeighborModel, Error, TEXT("Source anim %d is null."), i);
+				KMeansClusterResult |= EUpdateResult::ERROR;
+				return;
+			}
+		}
 
 		UNearestNeighborTrainingModel *TrainingModel = InitTrainingModel<UNearestNeighborTrainingModel>(this);
 		check(TrainingModel != nullptr);
-		TrainingModel->KmeansClusterPoses();
+		KMeansClusterResult |= TrainingModel->KmeansClusterPoses(NearestNeighborModel->KMeansPartId);
+		if (HasError(KMeansClusterResult))
+		{
+			return;
+		}
 		ResetSamplerData();
+		TArray<int32> KmeansResults = TrainingModel->KmeansResults;
+
+		const FString PackageName = GetTestMLDeformerComponent()->GetDeformerAsset()->GetPackage()->GetName();
+		const FString DirName = FPackageName::GetLongPackagePath(PackageName);
+		const FString FileName = FPaths::GetBaseFilename(PackageName);
+		const FString SavePath = FString::Printf(TEXT("%s/%s_PartId_%d"), *DirName, *FileName, NearestNeighborModel->KMeansPartId);
+
+		TPair<UAnimSequence*, uint8> AnimAndFlag = CreateAnimOfClusterCenters(SavePath, KmeansResults);
+		UAnimSequence* Anim = AnimAndFlag.Get<0>();
+		KMeansClusterResult |= AnimAndFlag.Get<1>();
+		if (HasError(KMeansClusterResult))
+		{
+			return;
+		}
+
+		UPackage* Package = Anim->GetPackage();
+		const bool bSaveSucced = UPackage::SavePackage(Package, Anim, *SavePath, FSavePackageArgs());
+
+		KmeansResults.Reset();
+	}
+
+	TPair<UAnimSequence*, uint8> FNearestNeighborEditorModel::CreateAnimOfClusterCenters(const FString& PackageName, const TArray<int32>& KmeansResults)
+	{
+		uint8 ReturnCode = EUpdateResult::SUCCESS;
+		UNearestNeighborModel *NearestNeighborModel = static_cast<UNearestNeighborModel*>(Model);
+
+		if (KmeansResults.Num() != NearestNeighborModel->NumClusters * 2)
+		{
+			UE_LOG(LogNearestNeighborModel, Error, TEXT("KmeansClusterPoses returned %d clusters whereas %d are expected."), KmeansResults.Num() / 2, NearestNeighborModel->NumClusters);
+			return TTuple<UAnimSequence*, uint8>(nullptr, EUpdateResult::ERROR);
+		}
+
+		const UAnimSequence* DefaultAnim = NearestNeighborModel->SourceAnims[0];
+
+		UAnimSequence* Anim = CreateObjectInstance<UAnimSequence>(PackageName);
+		if (Anim == nullptr)
+		{
+			return TTuple<UAnimSequence*, uint8>(nullptr, EUpdateResult::ERROR);
+		}
+
+		Anim->SetSkeleton(DefaultAnim->GetSkeleton());
+		IAnimationDataController& Controller = Anim->GetController();
+		Controller.OpenBracket(LOCTEXT("CreateNewAnim_Bracket", "Create New Anim"));
+		Controller.InitializeModel();
+		Anim->ResetAnimation();
+		Anim->SetPreviewMesh(NearestNeighborModel->GetSkeletalMesh());
+		const IAnimationDataModel* AnimData = Anim->GetDataModel();
+		const int32 NumKeys = NearestNeighborModel->NumClusters;
+		Controller.SetNumberOfFrames(NumKeys - 1);
+		Controller.SetFrameRate(FFrameRate(30, 1));
+
+		const TArray<FBoneAnimationTrack>& DefaultTracks = DefaultAnim->GetDataModel()->GetBoneAnimationTracks();
+		const int32 NumTracks = DefaultTracks.Num();
+		for (int32 TrackIndex = 0; TrackIndex < NumTracks; TrackIndex++)
+		{
+			const FName TrackName = DefaultTracks[TrackIndex].Name;
+			TArray<FVector3f> PosKeys;
+			TArray<FQuat4f> RotKeys;
+			TArray<FVector3f> ScaleKeys;
+			PosKeys.SetNum(NumKeys);
+			RotKeys.SetNum(NumKeys);
+			ScaleKeys.SetNum(NumKeys);
+			for (int32 KeyIndex = 0; KeyIndex < NumKeys; KeyIndex++)
+			{
+				const int32 PickedAnimId = KmeansResults[KeyIndex * 2];
+				const int32 PickedFrame = KmeansResults[KeyIndex * 2 + 1];
+				if (PickedAnimId < 0 || PickedAnimId >= NearestNeighborModel->SourceAnims.Num())
+				{
+					UE_LOG(LogNearestNeighborModel, Error, TEXT("CreateAnimOfClusterCenters: PickedAnimId %d is out of range."), PickedAnimId);
+					return TTuple<UAnimSequence*, uint8>(nullptr, EUpdateResult::ERROR);
+				}
+
+				const UAnimSequence* PickedAnim = NearestNeighborModel->SourceAnims[PickedAnimId];
+				if (PickedAnim == nullptr)
+				{
+					UE_LOG(LogNearestNeighborModel, Error, TEXT("CreateAnimOfClusterCenters: PickedAnim %d is null."), PickedAnimId);
+					return TTuple<UAnimSequence*, uint8>(nullptr, EUpdateResult::ERROR);
+				}
+				const FBoneAnimationTrack& Track = PickedAnim->GetDataModel()->GetBoneTrackByName(TrackName);
+				const FRawAnimSequenceTrack& RawTrack = Track.InternalTrackData;
+				if (PickedFrame < 0 || PickedFrame >= RawTrack.PosKeys.Num())
+				{
+					UE_LOG(LogNearestNeighborModel, Error, TEXT("CreateAnimOfClusterCenters: PickedFrame %d is out of range[0, %d)."), PickedFrame, RawTrack.PosKeys.Num());
+					return TTuple<UAnimSequence*, uint8>(nullptr, EUpdateResult::ERROR);
+				}
+				PosKeys[KeyIndex] = RawTrack.PosKeys[PickedFrame];
+				RotKeys[KeyIndex] = RawTrack.RotKeys[PickedFrame];
+				ScaleKeys[KeyIndex] = RawTrack.ScaleKeys[PickedFrame];
+			}
+			Controller.AddBoneTrack(TrackName);
+			Controller.SetBoneTrackKeys(TrackName, PosKeys, RotKeys, ScaleKeys);
+		}
+		Controller.NotifyPopulated();
+		Controller.CloseBracket();
+		return MakeTuple(Anim, ReturnCode);
 	}
 
 	void FNearestNeighborEditorModel::AddFloatArrayToDeltaArray(const TArray<float>& FloatArr, const TArray<uint32>& VertexMap, TArray<FVector3f>& DeltaArr, int32 DeltaArrayOffset, float ScaleFactor)
diff --git a/Engine/Plugins/Animation/MLDeformer/NearestNeighborModel/Source/NearestNeighborModelEditor/Private/NearestNeighborGeomCacheSampler.cpp b/Engine/Plugins/Animation/MLDeformer/NearestNeighborModel/Source/NearestNeighborModelEditor/Private/NearestNeighborGeomCacheSampler.cpp
--- a/Engine/Plugins/Animation/MLDeformer/NearestNeighborModel/Source/NearestNeighborModelEditor/Private/NearestNeighborGeomCacheSampler.cpp
+++ b/Engine/Plugins/Animation/MLDeformer/NearestNeighborModel/Source/NearestNeighborModelEditor/Private/NearestNeighborGeomCacheSampler.cpp
@@ -261,33 +261,63 @@ namespace UE::NearestNeighborModel
 		}
 	}
 
-	void FNearestNeighborGeomCacheSampler::SampleKMeansAnim(const int32 SkeletonId)
+	bool FNearestNeighborGeomCacheSampler::SampleKMeansAnim(const int32 AnimId)
 	{
 		UNearestNeighborModel* NearestNeighborModel = static_cast<UNearestNeighborModel*>(Model);
-		if (NearestNeighborModel && SkeletonId < NearestNeighborModel->SourceSkeletons.Num() && SkeletalMeshComponent)
+		if (NearestNeighborModel && AnimId < NearestNeighborModel->SourceAnims.Num() && SkeletalMeshComponent)
 		{
-			const TObjectPtr<UAnimSequence> AnimSequence = NearestNeighborModel->SourceSkeletons[SkeletonId];
+			const TObjectPtr<UAnimSequence> AnimSequence = NearestNeighborModel->SourceAnims[AnimId];
 			SkeletalMeshComponent->SetAnimationMode(EAnimationMode::AnimationSingleNode);
 			SkeletalMeshComponent->SetAnimation(AnimSequence);
 			SkeletalMeshComponent->SetPosition(0.0f);
 			SkeletalMeshComponent->SetPlayRate(1.0f);
 			SkeletalMeshComponent->Play(false);
 			SkeletalMeshComponent->RefreshBoneTransforms();
+			KMeansAnimId = AnimId;
+			return true;
+		}
+		else
+		{
+			return false;
 		}
 	}
 
-	void FNearestNeighborGeomCacheSampler::SampleKMeansFrame(const int32 Frame)
-	{
-		AnimFrameIndex = Frame;
-		SampleTime = GetTimeAtFrame(Frame);
+	// Write a function to get the animation of a skeletal mesh component
 
+	bool FNearestNeighborGeomCacheSampler::SampleKMeansFrame(const int32 Frame)
+	{
 		const USkeletalMesh* SkeletalMesh = SkeletalMeshComponent->GetSkeletalMeshAsset();
 		if (SkeletalMeshComponent && SkeletalMesh)
 		{
-			UpdateSkeletalMeshComponent();
-			UpdateBoneRotations();
-			UpdateCurveValues();
+			UNearestNeighborModel* NearestNeighborModel = static_cast<UNearestNeighborModel*>(Model);
+			const UAnimSequence* AnimSequence = NearestNeighborModel->SourceAnims[KMeansAnimId];
+			if (NearestNeighborModel->GetSkeletalMesh() == nullptr)
+			{
+				UE_LOG(LogNearestNeighborModel, Error, TEXT("SkeletalMesh is nullptr. Unable to sample KMeans frame."));
+				return false;
+			}
+
+			if (AnimSequence && Frame < AnimSequence->GetDataModel()->GetNumberOfKeys())
+			{
+				AnimFrameIndex = Frame;
+				SampleTime = GetTimeAtFrame(Frame);
+
+				UpdateSkeletalMeshComponent();
+				UpdateBoneRotations();
+				UpdateCurveValues();
+				return true;
+			}
+			else
+			{
+				UE_LOG(LogNearestNeighborModel, Error, TEXT("AnimSequence only has %d keys, but being sampled with key %d"), AnimSequence->GetDataModel()->GetNumberOfKeys(), Frame);
+				return false;
+			}
+		}
+		else
+		{
+			UE_LOG(LogNearestNeighborModel, Error, TEXT("KMeans: SkeletalMesh does not exist"));
 		}
+		return false;
 	}
 
 	uint8 FNearestNeighborGeomCacheSampler::GenerateMeshMappingIndices()
diff --git a/Engine/Plugins/Animation/MLDeformer/NearestNeighborModel/Source/NearestNeighborModelEditor/Private/NearestNeighborModelDetails.cpp b/Engine/Plugins/Animation/MLDeformer/NearestNeighborModel/Source/NearestNeighborModelEditor/Private/NearestNeighborModelDetails.cpp
--- a/Engine/Plugins/Animation/MLDeformer/NearestNeighborModel/Source/NearestNeighborModelEditor/Private/NearestNeighborModelDetails.cpp
+++ b/Engine/Plugins/Animation/MLDeformer/NearestNeighborModel/Source/NearestNeighborModelEditor/Private/NearestNeighborModelDetails.cpp
@@ -193,7 +193,7 @@ namespace UE::NearestNeighborModel
 
 		BuildSubMeshNames();
 		TSharedRef<IPropertyHandle> ClothPartDataPropertyHandle = DetailBuilder.GetProperty(UNearestNeighborModel::GetClothPartEditorDataPropertyName());
-		if (ClothPartDataPropertyHandle->AsArray().IsValid())
+		if (ClothPartDataPropertyHandle->AsArray().IsValid() && !SubMeshNames.IsEmpty())
 		{
 			TSharedRef<FDetailArrayBuilder> PropertyBuilder = MakeShared<FDetailArrayBuilder>(ClothPartDataPropertyHandle, true, false, true);
 			PropertyBuilder->OnGenerateArrayElementWidget(FOnGenerateArrayElementWidget::CreateSP(this, &FNearestNeighborModelDetails::GenerateClothPartElementWidget));
@@ -206,8 +206,9 @@ namespace UE::NearestNeighborModel
 		NearestNeighborCategoryBuilder->AddProperty(UNearestNeighborModel::GetUsePartOnlyMeshPropertyName());
 		NearestNeighborCategoryBuilder->AddProperty(UNearestNeighborModel::GetNearestNeighborDataPropertyName());
 
-		KMeansCategoryBuilder->AddProperty(GET_MEMBER_NAME_STRING_CHECKED(UNearestNeighborModel, SourceSkeletons));
+		KMeansCategoryBuilder->AddProperty(GET_MEMBER_NAME_STRING_CHECKED(UNearestNeighborModel, SourceAnims));
 		KMeansCategoryBuilder->AddProperty(GET_MEMBER_NAME_STRING_CHECKED(UNearestNeighborModel, NumClusters));
+		KMeansCategoryBuilder->AddProperty(GET_MEMBER_NAME_STRING_CHECKED(UNearestNeighborModel, KMeansPartId));
 		KMeansCategoryBuilder->AddCustomRow(FText::FromString(""))
 			.WholeRowContent()
 			[
@@ -217,6 +218,11 @@ namespace UE::NearestNeighborModel
 				.OnClicked_Lambda([this]
 				{
 					NearestNeighborEditorModel->KMeansClusterPoses();
+					EditorModel->GetEditor()->GetModelDetailsView()->ForceRefresh();
+					if (NearestNeighborEditorModel->GetKMeansClusterResult() == EUpdateResult::SUCCESS)
+					{
+						UE_LOG(LogNearestNeighborModel, Display, TEXT("Cluster succeeded."));
+					}
 					return FReply::Handled();
 				})
 			];
@@ -235,13 +241,18 @@ namespace UE::NearestNeighborModel
 				{
 					NearestNeighborEditorModel->OnMorphTargetUpdate();
 					EditorModel->GetEditor()->GetModelDetailsView()->ForceRefresh();
+					if (NearestNeighborEditorModel->GetMorphTargetUpdateResult() == EUpdateResult::SUCCESS)
+					{
+						UE_LOG(LogNearestNeighborModel, Display, TEXT("Update succeeded."));
+					}
 					return FReply::Handled();
 				})
 			];
-		AddUpdateResultText(MorphTargetCategoryBuilder, NearestNeighborEditorModel->GetMorphTargetUpdateResult());
+		AddActionResultText(MorphTargetCategoryBuilder, NearestNeighborEditorModel->GetMorphTargetUpdateResult(), TEXT("Update"));
+		AddActionResultText(KMeansCategoryBuilder, NearestNeighborEditorModel->GetKMeansClusterResult(), TEXT("KMeans"));
 	}
 
-	void FNearestNeighborModelDetails::AddUpdateResultText(IDetailCategoryBuilder* CategoryBuilder, uint8 Result)
+	void FNearestNeighborModelDetails::AddActionResultText(IDetailCategoryBuilder* CategoryBuilder, uint8 Result, const FString& ActionName)
 	{
 		CategoryBuilder->AddCustomRow(FText::FromString("UpdateResultError"))
 			.Visibility((Result & EUpdateResult::ERROR) != 0 ? EVisibility::Visible : EVisibility::Collapsed)
@@ -252,7 +263,7 @@ namespace UE::NearestNeighborModel
 				[
 					SNew(SWarningOrErrorBox)
 					.MessageStyle(EMessageStyle::Error)
-					.Message(FText::FromString(TEXT("Update failed with errors. Please check Output Log (LogNearestNeighborModel, LogPython) for details.")))
+					.Message(FText::FromString(ActionName + TEXT(" failed with errors. Please check Output Log (LogNearestNeighborModel, LogPython) for details.")))
 				]
 			];
 		CategoryBuilder->AddCustomRow(FText::FromString("UpdateResultWarning"))
@@ -264,13 +275,9 @@ namespace UE::NearestNeighborModel
 				[
 					SNew(SWarningOrErrorBox)
 					.MessageStyle(EMessageStyle::Warning)
-					.Message(FText::FromString(TEXT("Update finished with warnings. Please check Output Log (LogNearestNeighborModel, LogPython) for details.")))
+					.Message(FText::FromString(ActionName + TEXT(" finished with warnings. Please check Output Log (LogNearestNeighborModel, LogPython) for details.")))
 				]
 			];
-		if (Result == EUpdateResult::SUCCESS)
-		{
-			UE_LOG(LogNearestNeighborModel, Display, TEXT("Update succeeded."));
-		}
 	}
 
 	void FNearestNeighborModelDetails::BuildSubMeshNames()
diff --git a/Engine/Plugins/Animation/MLDeformer/NearestNeighborModel/Source/NearestNeighborModelEditor/Private/NearestNeighborTrainingModel.cpp b/Engine/Plugins/Animation/MLDeformer/NearestNeighborModel/Source/NearestNeighborModelEditor/Private/NearestNeighborTrainingModel.cpp
--- a/Engine/Plugins/Animation/MLDeformer/NearestNeighborModel/Source/NearestNeighborModelEditor/Private/NearestNeighborTrainingModel.cpp
+++ b/Engine/Plugins/Animation/MLDeformer/NearestNeighborModel/Source/NearestNeighborModelEditor/Private/NearestNeighborTrainingModel.cpp
@@ -52,30 +52,35 @@ int32 UNearestNeighborTrainingModel::GetPartNumNeighbors(const int32 PartId) con
 	return FMath::Min(NearestNeighborModel->GetNumNeighborsFromAnimSequence(PartId), NearestNeighborModel->GetNumNeighborsFromGeometryCache(PartId));
 }
 
-void UNearestNeighborTrainingModel::SampleKmeansAnim(const int32 SkeletonId)
+bool UNearestNeighborTrainingModel::SampleKmeansAnim(const int32 SkeletonId)
 {
 	FNearestNeighborGeomCacheSampler* Sampler = static_cast<FNearestNeighborGeomCacheSampler*>(EditorModel->GetSampler());
-	Sampler->SampleKMeansAnim(SkeletonId);
+	return Sampler->SampleKMeansAnim(SkeletonId);
 }
 
-void UNearestNeighborTrainingModel::SampleKmeansFrame(const int32 Frame)
+bool UNearestNeighborTrainingModel::SampleKmeansFrame(const int32 Frame)
 {
 	FNearestNeighborGeomCacheSampler* Sampler = static_cast<FNearestNeighborGeomCacheSampler*>(EditorModel->GetSampler());
-	Sampler->SampleKMeansFrame(Frame);
-	SampleBoneRotations = Sampler->GetBoneRotations();
+	const bool bSampleExist = Sampler->SampleKMeansFrame(Frame);
+	if (bSampleExist)
+	{
+		SampleBoneRotations = Sampler->GetBoneRotations();
+		return true;	
+	}
+	return false;
 }
 
 
 int32 UNearestNeighborTrainingModel::GetKmeansNumAnims() const
 {
-	return NearestNeighborModel->SourceSkeletons.Num();
+	return NearestNeighborModel->SourceAnims.Num();
 }
 
 int32 UNearestNeighborTrainingModel::GetKmeansAnimNumFrames(const int32 SkeletonId) const
 {
-	if (SkeletonId < NearestNeighborModel->SourceSkeletons.Num())
+	if (SkeletonId < NearestNeighborModel->SourceAnims.Num())
 	{
-		return NearestNeighborModel->SourceSkeletons[SkeletonId]->GetDataModel()->GetNumberOfFrames();
+		return NearestNeighborModel->SourceAnims[SkeletonId]->GetDataModel()->GetNumberOfFrames();
 	}
 	return 0;
 }
diff --git a/Engine/Plugins/Animation/MLDeformer/NearestNeighborModel/Source/NearestNeighborModelEditor/Public/NearestNeighborEditorModel.h b/Engine/Plugins/Animation/MLDeformer/NearestNeighborModel/Source/NearestNeighborModelEditor/Public/NearestNeighborEditorModel.h
--- a/Engine/Plugins/Animation/MLDeformer/NearestNeighborModel/Source/NearestNeighborModelEditor/Public/NearestNeighborEditorModel.h
+++ b/Engine/Plugins/Animation/MLDeformer/NearestNeighborModel/Source/NearestNeighborModelEditor/Public/NearestNeighborEditorModel.h
@@ -61,6 +61,8 @@ namespace UE::NearestNeighborModel
 
 		void UpdateNearestNeighborActors();
 		void KMeansClusterPoses();
+		TPair<UAnimSequence*, uint8> CreateAnimOfClusterCenters(const FString& PackageName, const TArray<int32>& KmeansResults);
+		uint8 GetKMeansClusterResult() const { return KMeansClusterResult; }
 
 		uint8 InitMorphTargets();
 		void RefreshMorphTargets();
@@ -99,9 +101,8 @@ namespace UE::NearestNeighborModel
 
 	private:
 		UWorld* EditorWorld = nullptr;
-		uint8 ClothPartUpdateResult = 0;
-		uint8 NearestNeighborUpdateResult = 0;
 		uint8 MorphTargetUpdateResult = 0;
+		uint8 KMeansClusterResult = 0;
 		int32 NumTrainingFramesOverride = -1;
 	};
 }	// namespace UE::NearestNeighborModel
diff --git a/Engine/Plugins/Animation/MLDeformer/NearestNeighborModel/Source/NearestNeighborModelEditor/Public/NearestNeighborGeomCacheSampler.h b/Engine/Plugins/Animation/MLDeformer/NearestNeighborModel/Source/NearestNeighborModelEditor/Public/NearestNeighborGeomCacheSampler.h
--- a/Engine/Plugins/Animation/MLDeformer/NearestNeighborModel/Source/NearestNeighborModelEditor/Public/NearestNeighborGeomCacheSampler.h
+++ b/Engine/Plugins/Animation/MLDeformer/NearestNeighborModel/Source/NearestNeighborModelEditor/Public/NearestNeighborGeomCacheSampler.h
@@ -11,8 +11,8 @@ namespace UE::NearestNeighborModel
 	{
 	public:
 		virtual uint8 SamplePart(int32 InAnimFrameIndex, int32 PartId);
-		virtual void SampleKMeansAnim(const int32 SkeletonId);
-		virtual void SampleKMeansFrame(const int32 Frame);
+		virtual bool SampleKMeansAnim(const int32 SkeletonId);
+		virtual bool SampleKMeansFrame(const int32 Frame);
 		const TArray<float>& GetPartVertexDeltas() const { return PartVertexDeltas; }
 		uint8 GeneratePartMeshMappings(const TArray<uint32>& VertexMap, bool bUsePartOnlyMesh);
 		uint8 GenerateMeshMappingIndices();
@@ -22,5 +22,6 @@ namespace UE::NearestNeighborModel
 	protected:
 		TArray<float> PartVertexDeltas;
 		TArray<int32> MeshMappingIndices;
+		int32 KMeansAnimId;
 	};
 }
diff --git a/Engine/Plugins/Animation/MLDeformer/NearestNeighborModel/Source/NearestNeighborModelEditor/Public/NearestNeighborModelDetails.h b/Engine/Plugins/Animation/MLDeformer/NearestNeighborModel/Source/NearestNeighborModelEditor/Public/NearestNeighborModelDetails.h
--- a/Engine/Plugins/Animation/MLDeformer/NearestNeighborModel/Source/NearestNeighborModelEditor/Public/NearestNeighborModelDetails.h
+++ b/Engine/Plugins/Animation/MLDeformer/NearestNeighborModel/Source/NearestNeighborModelEditor/Public/NearestNeighborModelDetails.h
@@ -8,6 +8,7 @@
 class UNearestNeighborModel;
 class IPropertyHandle;
 class IDetailChildrenBuilder;
+class FDetailWidgetRow;
 
 namespace UE::NearestNeighborModel
 {
@@ -27,7 +28,7 @@ namespace UE::NearestNeighborModel
 		virtual void CreateCategories() override;
 
 		virtual bool UpdateMemberPointers(const TArray<TWeakObjectPtr<UObject>>& Objects) override;
-		virtual void AddUpdateResultText(IDetailCategoryBuilder* CategoryBuilder, uint8 Result);
+		virtual void AddActionResultText(IDetailCategoryBuilder* CategoryBuilder, uint8 Result, const FString& ActionName);
 
 		void GenerateClothPartElementWidget(TSharedRef<IPropertyHandle> PropertyHandle, int32 ArrayIndex, IDetailChildrenBuilder& ChildrenBuilder);
 
diff --git a/Engine/Plugins/Animation/MLDeformer/NearestNeighborModel/Source/NearestNeighborModelEditor/Public/NearestNeighborTrainingModel.h b/Engine/Plugins/Animation/MLDeformer/NearestNeighborModel/Source/NearestNeighborModelEditor/Public/NearestNeighborTrainingModel.h
--- a/Engine/Plugins/Animation/MLDeformer/NearestNeighborModel/Source/NearestNeighborModelEditor/Public/NearestNeighborTrainingModel.h
+++ b/Engine/Plugins/Animation/MLDeformer/NearestNeighborModel/Source/NearestNeighborModelEditor/Public/NearestNeighborTrainingModel.h
@@ -47,13 +47,13 @@ public:
 	int32 GetPartNumNeighbors(const int32 PartId) const;
 
 	UFUNCTION(BlueprintImplementableEvent, Category = "Python")
-	void KmeansClusterPoses() const;
+	int32 KmeansClusterPoses(const int32 PartId) const;
 
 	UFUNCTION(BlueprintCallable, Category = "Python")
-	void SampleKmeansAnim(const int32 SkeletonId);
+	bool SampleKmeansAnim(const int32 SkeletonId);
 
 	UFUNCTION(BlueprintCallable, Category = "Python")
-	void SampleKmeansFrame(const int32 Frame);
+	bool SampleKmeansFrame(const int32 Frame);
 
 	UFUNCTION(BlueprintPure, Category = "Python")
 	int32 GetKmeansNumAnims() const;
@@ -69,5 +69,8 @@ public:
 	UPROPERTY(EditAnywhere, BlueprintReadWrite, Category = "Training Data")
 	TArray<float> PartSampleDeltas;
 
+	UPROPERTY(EditAnywhere, BlueprintReadWrite, Category = "Training Data")
+	TArray<int32> KmeansResults;
+
 	UNearestNeighborModel* NearestNeighborModel = nullptr;
 };
