NearestNeighborModel: fix static analyzer warnings

diff --git a/Engine/Plugins/Animation/MLDeformer/NearestNeighborModel/Source/NearestNeighborModelEditor/Private/NearestNeighborGeomCacheSampler.cpp b/Engine/Plugins/Animation/MLDeformer/NearestNeighborModel/Source/NearestNeighborModelEditor/Private/NearestNeighborGeomCacheSampler.cpp
--- a/Engine/Plugins/Animation/MLDeformer/NearestNeighborModel/Source/NearestNeighborModelEditor/Private/NearestNeighborGeomCacheSampler.cpp
+++ b/Engine/Plugins/Animation/MLDeformer/NearestNeighborModel/Source/NearestNeighborModelEditor/Private/NearestNeighborGeomCacheSampler.cpp
@@ -297,21 +297,30 @@ namespace UE::NearestNeighborModel
 				return false;
 			}
 
-			if (AnimSequence && Frame < AnimSequence->GetDataModel()->GetNumberOfKeys())
+			if (AnimSequence)
 			{
-				AnimFrameIndex = Frame;
-				SampleTime = GetTimeAtFrame(Frame);
+				if (Frame < AnimSequence->GetDataModel()->GetNumberOfKeys())
+				{
+					AnimFrameIndex = Frame;
+					SampleTime = GetTimeAtFrame(Frame);
 
-				UpdateSkeletalMeshComponent();
-				UpdateBoneRotations();
-				UpdateCurveValues();
-				return true;
+					UpdateSkeletalMeshComponent();
+					UpdateBoneRotations();
+					UpdateCurveValues();
+					return true;
+				}
+				else
+				{
+					UE_LOG(LogNearestNeighborModel, Error, TEXT("AnimSequence only has %d keys, but being sampled with key %d"), AnimSequence->GetDataModel()->GetNumberOfKeys(), Frame);
+					return false;
+				}
 			}
 			else
 			{
-				UE_LOG(LogNearestNeighborModel, Error, TEXT("AnimSequence only has %d keys, but being sampled with key %d"), AnimSequence->GetDataModel()->GetNumberOfKeys(), Frame);
+				UE_LOG(LogNearestNeighborModel, Error, TEXT("AnimSequence %d is nullptr. Unable to sample KMeans frame."), KMeansAnimId);
 				return false;
 			}
+
 		}
 		else
 		{
