Nearest neighbor model: show index number for nearest neighbor editor

diff --git a/Engine/Plugins/Animation/MLDeformer/NearestNeighborModel/Source/NearestNeighborModel/Public/NearestNeighborModelInstance.h b/Engine/Plugins/Animation/MLDeformer/NearestNeighborModel/Source/NearestNeighborModel/Public/NearestNeighborModelInstance.h
--- a/Engine/Plugins/Animation/MLDeformer/NearestNeighborModel/Source/NearestNeighborModel/Public/NearestNeighborModelInstance.h
+++ b/Engine/Plugins/Animation/MLDeformer/NearestNeighborModel/Source/NearestNeighborModel/Public/NearestNeighborModelInstance.h
@@ -18,6 +18,7 @@ public:
     virtual bool SetupInputs() override;
 
 #if WITH_EDITORONLY_DATA
+    const TArray<uint32>& GetNearestNeighborIds() const { return NearestNeighborIds; }
     uint32 NearestNeighborId(int32 PartId) const { return NearestNeighborIds[PartId]; }
     int32 NeighborIdNum() const { return NearestNeighborIds.Num(); }
 #endif
diff --git a/Engine/Plugins/Animation/MLDeformer/NearestNeighborModel/Source/NearestNeighborModel/Public/NearestNeighborModelVizSettings.h b/Engine/Plugins/Animation/MLDeformer/NearestNeighborModel/Source/NearestNeighborModel/Public/NearestNeighborModelVizSettings.h
--- a/Engine/Plugins/Animation/MLDeformer/NearestNeighborModel/Source/NearestNeighborModel/Public/NearestNeighborModelVizSettings.h
+++ b/Engine/Plugins/Animation/MLDeformer/NearestNeighborModel/Source/NearestNeighborModel/Public/NearestNeighborModelVizSettings.h
@@ -22,8 +22,14 @@ public:
 	float GetNearestNeighborActorsOffset() const { return NearestNeighborActorsOffset; }
 	static FName GetNearestNeighborActorsOffsetPropertyName() { return GET_MEMBER_NAME_CHECKED(UNearestNeighborModelVizSettings, NearestNeighborActorsOffset); }
 
+	static FName GetNearestNeighborIdsPropertyName() { return GET_MEMBER_NAME_CHECKED(UNearestNeighborModelVizSettings, NearestNeighborIds); }
+	void SetNearestNeighborIds(const TArray<uint32>& InNearestNeighborIds) { NearestNeighborIds = InNearestNeighborIds; }
+
 protected:
 	UPROPERTY(EditAnywhere, Category = "Live Settings")
 	float NearestNeighborActorsOffset = 2.0f;
+
+	UPROPERTY(VisibleAnywhere, Category = "Live Settings")
+	TArray<uint32> NearestNeighborIds;
 #endif
 };
diff --git a/Engine/Plugins/Animation/MLDeformer/NearestNeighborModel/Source/NearestNeighborModelEditor/Private/NearestNeighborEditorModel.cpp b/Engine/Plugins/Animation/MLDeformer/NearestNeighborModel/Source/NearestNeighborModelEditor/Private/NearestNeighborEditorModel.cpp
--- a/Engine/Plugins/Animation/MLDeformer/NearestNeighborModel/Source/NearestNeighborModelEditor/Private/NearestNeighborEditorModel.cpp
+++ b/Engine/Plugins/Animation/MLDeformer/NearestNeighborModel/Source/NearestNeighborModelEditor/Private/NearestNeighborEditorModel.cpp
@@ -134,7 +134,7 @@ namespace UE::NearestNeighborModel
 		FMLDeformerEditorModel::Tick(ViewportClient, DeltaTime);
 		if (!NearestNeighborActors.IsEmpty())
 		{
-			const UNearestNeighborModelVizSettings* NNViz = GetNearestNeighborModelVizSettings();
+			UNearestNeighborModelVizSettings* NNViz = GetNearestNeighborModelVizSettings();
 			const float Offset = NNViz->GetNearestNeighborActorsOffset();
 			for (FNearestNeighborEditorModelActor* NearestNeighborActor : NearestNeighborActors)
 			{
@@ -144,6 +144,15 @@ namespace UE::NearestNeighborModel
 					NearestNeighborActor->SetMeshOffsetFactor(Offset);
 				}
 			}
+			const UMLDeformerComponent* MLDeformerComponent = GetTestMLDeformerComponent();
+			if (MLDeformerComponent)
+			{
+				const UNearestNeighborModelInstance* ModelInstance = static_cast<UNearestNeighborModelInstance*>(MLDeformerComponent->GetModelInstance());
+				if (ModelInstance)
+				{
+					NNViz->SetNearestNeighborIds(ModelInstance->GetNearestNeighborIds());
+				}
+			}
 		}
 	}
 
diff --git a/Engine/Plugins/Animation/MLDeformer/NearestNeighborModel/Source/NearestNeighborModelEditor/Private/NearestNeighborModelVizSettingsDetails.cpp b/Engine/Plugins/Animation/MLDeformer/NearestNeighborModel/Source/NearestNeighborModelEditor/Private/NearestNeighborModelVizSettingsDetails.cpp
--- a/Engine/Plugins/Animation/MLDeformer/NearestNeighborModel/Source/NearestNeighborModelEditor/Private/NearestNeighborModelVizSettingsDetails.cpp
+++ b/Engine/Plugins/Animation/MLDeformer/NearestNeighborModel/Source/NearestNeighborModelEditor/Private/NearestNeighborModelVizSettingsDetails.cpp
@@ -18,6 +18,7 @@ namespace UE::NearestNeighborModel
 		FMLDeformerMorphModelVizSettingsDetails::AddAdditionalSettings();
 		IDetailGroup& NNGroup = LiveSettingsCategory->AddGroup("Nearest Neighbor", LOCTEXT("NearestNeighborLabel", "Nearest Neighbor"), false, true);
 		NNGroup.AddPropertyRow(DetailLayoutBuilder->GetProperty(UNearestNeighborModelVizSettings::GetNearestNeighborActorsOffsetPropertyName(), UNearestNeighborModelVizSettings::StaticClass()));
+		NNGroup.AddPropertyRow(DetailLayoutBuilder->GetProperty(UNearestNeighborModelVizSettings::GetNearestNeighborIdsPropertyName(), UNearestNeighborModelVizSettings::StaticClass()));
 	}
 
 };
