NearestNeighborModel: Enable NearestNeighborModel to use optimized CPU network
Add serlialization code to read onnx model weights in python
Merge multiple linear layers into one for performance optimization

diff --git a/Engine/Plugins/Animation/MLDeformer/NearestNeighborModel/Source/NearestNeighborModel/Private/NearestNeighborModel.cpp b/Engine/Plugins/Animation/MLDeformer/NearestNeighborModel/Source/NearestNeighborModel/Private/NearestNeighborModel.cpp
--- a/Engine/Plugins/Animation/MLDeformer/NearestNeighborModel/Source/NearestNeighborModel/Private/NearestNeighborModel.cpp
+++ b/Engine/Plugins/Animation/MLDeformer/NearestNeighborModel/Source/NearestNeighborModel/Private/NearestNeighborModel.cpp
@@ -3,6 +3,8 @@
 #include "NearestNeighborModel.h"
 #include "NearestNeighborModelInputInfo.h"
 #include "NearestNeighborModelInstance.h"
+#include "NearestNeighborOptimizedNetwork.h"
+#include "NearestNeighborOptimizedNetworkLoader.h"
 #include "MLDeformerComponent.h"
 #include "Engine/SkeletalMesh.h"
 #include "GeometryCache.h"
@@ -29,6 +31,16 @@ namespace UE::NearestNeighborModel
 }
 IMPLEMENT_MODULE(UE::NearestNeighborModel::FNearestNeighborModelModule, NearestNeighborModel)
 
+namespace UE::NearestNeighborModel
+{
+	bool bNearestNeighborModelUseOptimizedNetwork = true;
+	FAutoConsoleVariableRef CVarNearestNeighborModelUseOptimizedNetwork(
+		TEXT("p.NearestNeighborModel.UseOptimizedNetwork"),
+		bNearestNeighborModelUseOptimizedNetwork,
+		TEXT("Whether to use the optimized network.")
+	);
+
+}
 
 UNearestNeighborModel::UNearestNeighborModel(const FObjectInitializer& ObjectInitializer)
 	: Super(ObjectInitializer)
@@ -52,6 +64,28 @@ UMLDeformerModelInstance* UNearestNeighborModel::CreateModelInstance(UMLDeformer
 	return NewObject<UNearestNeighborModelInstance>(Component);
 }
 
+void UNearestNeighborModel::Serialize(FArchive& Archive)
+{
+	Super::Serialize(Archive);
+	if (Archive.IsCooking())
+	{
+		if (DoesUseOptimizedNetwork())
+		{
+			if (GetNeuralNetwork() != nullptr)
+			{
+				SetNeuralNetwork(nullptr);
+			}
+		}
+		else
+		{
+			if (GetOptimizedNetwork() != nullptr)
+			{
+				SetOptimizedNetwork(nullptr);
+			}
+		}
+	}
+}
+
 void UNearestNeighborModel::PostLoad()
 {
 	Super::PostLoad();
@@ -419,4 +453,68 @@ bool UNearestNeighborModel::CheckPCAData(int32 PartId) const
 	const FClothPartData& Data = ClothPartData[PartId];
 	return Data.VertexMap.Num() > 0 && Data.PCABasis.Num() == Data.VertexMap.Num() * 3 * Data.PCACoeffNum;
 }
+
+void UNearestNeighborModel::SetOptimizedNetwork(UNearestNeighborOptimizedNetwork* InOptimizedNetwork)
+{
+	GetNeuralNetworkModifyDelegate().Broadcast();
+	OptimizedNetwork = InOptimizedNetwork;
+}
+
+bool UNearestNeighborModel::DoesUseOptimizedNetwork() const
+{
+	return UE::NearestNeighborModel::bNearestNeighborModelUseOptimizedNetwork;
+}
+
+/**
+ * Get derived class default object (CDO)
+ * @return A pointer to a CDO of the derived python class. 
+ */
+template<class T>
+T* GetDerivedCDO()
+{
+	TArray<UClass*> Classes;
+	GetDerivedClasses(T::StaticClass(), Classes);
+	if (Classes.IsEmpty())
+	{
+		return nullptr;
+	}
+
+	T* Object = Cast<T>(Classes.Last()->GetDefaultObject());
+	return Object;
+}
+
+bool UNearestNeighborModel::LoadOptimizedNetwork(const FString& OnnxPath)
+{
+	const FString OnnxFile = FPaths::ConvertRelativePathToFull(OnnxPath);
+	if (FPaths::FileExists(OnnxFile))
+	{
+		UE_LOG(LogNearestNeighborModel, Display, TEXT("Loading Onnx file '%s'..."), *OnnxFile);
+		UNearestNeighborOptimizedNetwork* Result = NewObject<UNearestNeighborOptimizedNetwork>(this, UNearestNeighborOptimizedNetwork::StaticClass());
+
+		UNearestNeighborOptimizedNetworkLoader* Loader = GetDerivedCDO<UNearestNeighborOptimizedNetworkLoader>();
+		if (Loader == nullptr)
+		{
+			return false;
+		}
+		Loader->SetOptimizedNetwork(Result);
+		const bool bSuccess = Loader->LoadOptimizedNetwork(OnnxFile);
+		if (bSuccess)
+		{
+			SetOptimizedNetwork(Result);
+			return true;
+		}
+	}
+	else
+	{
+		UE_LOG(LogMLDeformer, Error, TEXT("Onnx file '%s' does not exist!"), *OnnxFile);
+	}
+
+	return false;
+}
+
+int32 UNearestNeighborModel::GetOptimizedNetworkNumOutputs() const
+{
+	return OptimizedNetwork ? OptimizedNetwork->GetNumOutputs() : 0;
+}
+
 #undef LOCTEXT_NAMESPACE
diff --git a/Engine/Plugins/Animation/MLDeformer/NearestNeighborModel/Source/NearestNeighborModel/Private/NearestNeighborModelInstance.cpp b/Engine/Plugins/Animation/MLDeformer/NearestNeighborModel/Source/NearestNeighborModel/Private/NearestNeighborModelInstance.cpp
--- a/Engine/Plugins/Animation/MLDeformer/NearestNeighborModel/Source/NearestNeighborModel/Private/NearestNeighborModelInstance.cpp
+++ b/Engine/Plugins/Animation/MLDeformer/NearestNeighborModel/Source/NearestNeighborModel/Private/NearestNeighborModelInstance.cpp
@@ -4,7 +4,10 @@
 #include "NearestNeighborModel.h"
 #include "NearestNeighborModelInputInfo.h"
 #include "NeuralNetwork.h"
+#include "NearestNeighborOptimizedNetwork.h"
+#include "MLDeformerAsset.h"
 #include "Components/ExternalMorphSet.h"
+#include "Engine/SkeletalMesh.h"
 #include "Components/SkeletalMeshComponent.h"
 
 #include UE_INLINE_GENERATED_CPP_BY_NAME(NearestNeighborModelInstance)
@@ -39,22 +42,58 @@ int64 UNearestNeighborModelInstance::SetBoneTransforms(float* OutputBuffer, int6
 
 bool UNearestNeighborModelInstance::SetupInputs()
 {
-    const bool bSuccess = Super::SetupInputs();
-
-    if (bSuccess)
-    {
-	    UNeuralNetwork* NeuralNetwork = Model->GetNeuralNetwork();
-	    if (NeuralNetwork)
-	    {
-			UNearestNeighborModel* NearestNeighborModel = static_cast<UNearestNeighborModel*>(Model.Get());
-	    	float* InputDataPointer = static_cast<float*>(NeuralNetwork->GetInputDataPointerMutableForContext(NeuralNetworkInferenceHandle));
-	    	const int64 NumNeuralNetInputs = NeuralNetwork->GetInputTensorForContext(NeuralNetworkInferenceHandle).Num();
-		    // Clip network inputs based on training set min and max
-	    	NearestNeighborModel->ClipInputs(InputDataPointer, NumNeuralNetInputs);
-	    	return true;
-	    }
-    }
+	const UNearestNeighborModel* NearestNeighborModel = Cast<UNearestNeighborModel>(Model);
+	check(NearestNeighborModel);
+	if (NearestNeighborModel && NearestNeighborModel->DoesUseOptimizedNetwork())
+	{
+		if (NearestNeighborModel == nullptr ||
+			SkeletalMeshComponent == nullptr ||
+			SkeletalMeshComponent->GetSkeletalMeshAsset() == nullptr ||
+			!bIsCompatible)
+		{
+			return false;
+		}
+
+		const UNearestNeighborOptimizedNetwork* OptimizedNetwork = NearestNeighborModel->GetOptimizedNetwork();
+		if (!OptimizedNetwork || !OptimizedNetworkInstance)
+		{
+			return false;
+		}
+
+		int32 NumNeuralNetInputs = 0;
+		float* InputDataPointer = nullptr;
+		GetInputDataPointer(InputDataPointer, NumNeuralNetInputs);
+
+		// If the neural network expects a different number of inputs, do nothing.
+		const int32 NumFloatsPerBone = NearestNeighborModel->GetNumFloatsPerBone();
+		const int32 NumFloatsPerCurve = NearestNeighborModel->GetNumFloatsPerCurve();
+		const int64 NumDeformerAssetInputs = NearestNeighborModel->GetInputInfo()->CalcNumNeuralNetInputs(NumFloatsPerBone, NumFloatsPerCurve);
+		if (NumNeuralNetInputs != NumDeformerAssetInputs)
+		{
+			return false;
+		}
 
+		const int64 NumFloatsWritten = SetNeuralNetworkInputValues(InputDataPointer, NumNeuralNetInputs);
+		check(NumFloatsWritten == NumNeuralNetInputs);
+	}
+	else
+	{
+		const bool bSuccess = Super::SetupInputs();
+		if (!bSuccess)
+		{
+			return false;
+		}
+	}
+
+	float* InputDataPointer = nullptr;
+	int32 NumNeuralNetInputs = 0;
+	GetInputDataPointer(InputDataPointer, NumNeuralNetInputs);
+	if (InputDataPointer)
+	{
+		// Clip network inputs based on training set min and max values.
+		NearestNeighborModel->ClipInputs(InputDataPointer, NumNeuralNetInputs);
+		return true;
+	}
     return false;
 }
 
@@ -62,11 +101,35 @@ void UNearestNeighborModelInstance::Init(USkeletalMeshComponent* SkelMeshCompone
 {
 	UMLDeformerMorphModelInstance::Init(SkelMeshComponent);
 	InitPreviousWeights();
+	const UNearestNeighborModel* NearestNeighborModel = Cast<UNearestNeighborModel>(Model);
+	check(NearestNeighborModel);
+	if (NearestNeighborModel->DoesUseOptimizedNetwork())
+	{
+		InitOptimizedNetworkInstance();
+	}
+}
+
+void UNearestNeighborModelInstance::InitOptimizedNetworkInstance()
+{
+	const UNearestNeighborModel* NearestNeighborModel = Cast<UNearestNeighborModel>(Model);
+	check(NearestNeighborModel);
+	UNearestNeighborOptimizedNetwork* Network = NearestNeighborModel->GetOptimizedNetwork();
+	OptimizedNetworkInstance = Network ? Network->CreateInstance() : nullptr;
 }
 
 void UNearestNeighborModelInstance::Execute(float ModelWeight)
 {
-	Super::Execute(ModelWeight);
+	const UNearestNeighborModel* NearestNeighborModel = Cast<UNearestNeighborModel>(Model);
+	check(NearestNeighborModel);
+	if (NearestNeighborModel->DoesUseOptimizedNetwork())
+	{
+		UNearestNeighborOptimizedNetwork* OptimizedNetwork = NearestNeighborModel->GetOptimizedNetwork();
+		OptimizedNetworkInstance->Run();
+	}
+	else
+	{
+		Super::Execute(ModelWeight);
+	}
 	RunNearestNeighborModel(ModelWeight);
 }
 
@@ -102,11 +165,12 @@ void UNearestNeighborModelInstance::RunNearestNeighborModel(float ModelWeight)
 	NearestNeighborIds.SetNum(NearestNeighborModel->GetNumParts());
 #endif
 
-	const UNeuralNetwork* NeuralNetwork = Model->GetNeuralNetwork();
-	if (NeuralNetwork)
+	float* OutputTensorData = nullptr;
+	int32 NumNetworkWeights = 0;
+	GetOutputDataPointer(OutputTensorData, NumNetworkWeights);
+
+	if (OutputTensorData)
 	{
-		const FNeuralTensor& OutputTensor = NeuralNetwork->GetOutputTensorForContext(NeuralNetworkInferenceHandle);
-		const int32 NumNetworkWeights = OutputTensor.Num();
 		const int32 NumMorphTargets = WeightData->Weights.Num();
 		const int32 TotalNumNeighbors = NearestNeighborModel->GetTotalNumNeighbors();
 		if (NumMorphTargets >= NumNetworkWeights + 1 + TotalNumNeighbors)
@@ -116,14 +180,14 @@ void UNearestNeighborModelInstance::RunNearestNeighborModel(float ModelWeight)
 			// Update all generated morph target weights with the values calculated by our neural network.
 			for (int32 MorphIndex = 0; MorphIndex < NumNetworkWeights; ++MorphIndex)
 			{
-				const float W = OutputTensor.At<float>(MorphIndex) * ModelWeight;
+				const float W = OutputTensorData[MorphIndex] * ModelWeight;
 				UpdateWeight(WeightData->Weights, MorphIndex + 1, W);
 			}
 
 			int32 NeighborOffset = NumNetworkWeights + 1;
 			for (int32 PartId = 0; PartId < NearestNeighborModel->GetNumParts(); PartId++)
 			{
-				const int32 NearestNeighborId = FindNearestNeighbor(OutputTensor, PartId);
+				const int32 NearestNeighborId = FindNearestNeighbor(OutputTensorData, PartId);
 	#if WITH_EDITORONLY_DATA
 				NearestNeighborIds[PartId] = NearestNeighborId;
 	#endif
@@ -158,9 +222,10 @@ void UNearestNeighborModelInstance::UpdateWeight(TArray<float>& MorphWeights, in
 	PreviousWeights[Index] = NewW; 
 }
 
-int32 UNearestNeighborModelInstance::FindNearestNeighbor(const FNeuralTensor& PCACoeffTensor, int32 PartId)
+int32 UNearestNeighborModelInstance::FindNearestNeighbor(const float* PCAData, int32 PartId)
 {
 	const UNearestNeighborModel* NearestNeighborModel = Cast<UNearestNeighborModel>(Model);
+	check(NearestNeighborModel);
 	const TArray<float> &NeighborCoeffs = NearestNeighborModel->NeighborCoeffs(PartId);
 	float MinD2 = std::numeric_limits<float>::max();
 	int32 MinId = -1;
@@ -172,7 +237,7 @@ int32 UNearestNeighborModelInstance::FindNearestNeighbor(const FNeuralTensor& PC
 		const int32 NumNeighbors = NearestNeighborModel->GetNumNeighbors(PartId);
 		for (int32 CoeffId = 0; CoeffId < NumCoeffs; CoeffId++)
 		{
-			float Coeff = PCACoeffTensor.At<float>(CoeffStart + CoeffId);
+			float Coeff = PCAData[CoeffStart + CoeffId];
 			float NeighborCoeff = NeighborCoeffs[CoeffId * NumNeighbors + NeighborId];
 			float D = Coeff - NeighborCoeff;
 			D2 += D * D;
@@ -195,3 +260,119 @@ void UNearestNeighborModelInstance::InitPreviousWeights()
 		PreviousWeights.SetNumZeroed(NumWeights);
 	}
 }
+
+void UNearestNeighborModelInstance::GetInputDataPointer(float*& OutInputData, int32& OutNumInputFloats) const
+{
+	const UNearestNeighborModel* NearestNeighborModel = Cast<UNearestNeighborModel>(Model);
+	check(NearestNeighborModel);
+	if (NearestNeighborModel->DoesUseOptimizedNetwork())
+	{
+		const UNearestNeighborOptimizedNetwork* OptimizedNetwork = NearestNeighborModel->GetOptimizedNetwork();
+		if (OptimizedNetwork && OptimizedNetworkInstance)
+		{
+			OutInputData = OptimizedNetworkInstance->GetInputs().GetData();
+			OutNumInputFloats = OptimizedNetwork->GetNumInputs();
+			return;
+		}
+	}
+	else
+	{
+		UNeuralNetwork* NeuralNetwork = NearestNeighborModel->GetNeuralNetwork();
+		if (NeuralNetwork)
+		{
+			OutInputData = static_cast<float*>(NeuralNetwork->GetInputDataPointerMutableForContext(NeuralNetworkInferenceHandle));
+			OutNumInputFloats = NeuralNetwork->GetInputTensorForContext(NeuralNetworkInferenceHandle).Num();
+			return;
+		}
+	}
+	OutInputData = nullptr;
+	OutNumInputFloats = 0;
+}
+
+void UNearestNeighborModelInstance::GetOutputDataPointer(float*& OutOutputData, int32& OutNumOutputFloats) const
+{
+	const UNearestNeighborModel* NearestNeighborModel = Cast<UNearestNeighborModel>(Model);
+	check(NearestNeighborModel);
+	if (NearestNeighborModel->DoesUseOptimizedNetwork())
+	{
+		const UNearestNeighborOptimizedNetwork* OptimizedNetwork = NearestNeighborModel->GetOptimizedNetwork();
+		if (OptimizedNetwork && OptimizedNetworkInstance)
+		{
+			OutOutputData = OptimizedNetworkInstance->GetOutputs().GetData();
+			OutNumOutputFloats = OptimizedNetwork->GetNumOutputs();
+			return;
+		}
+	}
+	else
+	{
+		UNeuralNetwork* NeuralNetwork = NearestNeighborModel->GetNeuralNetwork();
+		if (NeuralNetwork)
+		{
+			const FNeuralTensor& OutputTensor = NeuralNetwork->GetOutputTensorForContext(NeuralNetworkInferenceHandle);
+			OutNumOutputFloats = OutputTensor.Num();
+			OutOutputData = (float*)OutputTensor.GetDataCasted<float>();
+			return;
+		}
+	}
+	OutOutputData = nullptr;
+	OutNumOutputFloats = 0;
+}
+
+FString UNearestNeighborModelInstance::CheckCompatibility(USkeletalMeshComponent* InSkelMeshComponent, bool LogIssues)
+{
+	ErrorText = FString();
+
+	// If we're not compatible, generate a compatibility string.
+	USkeletalMesh* SkelMesh = InSkelMeshComponent ? InSkelMeshComponent->GetSkeletalMeshAsset() : nullptr;
+	UMLDeformerInputInfo* InputInfo = Model->GetInputInfo();
+	if (SkelMesh && !InputInfo->IsCompatible(SkelMesh) && Model->GetDeformerAsset())
+	{
+		ErrorText += InputInfo->GenerateCompatibilityErrorString(SkelMesh);
+		ErrorText += "\n";
+		if (LogIssues)
+		{
+			UE_LOG(LogMLDeformer, Error, TEXT("ML Deformer '%s' isn't compatible with Skeletal Mesh '%s'.\nReason(s):\n%s"), 
+				*Model->GetDeformerAsset()->GetName(), 
+				*SkelMesh->GetName(), 
+				*ErrorText);
+		}
+	}
+
+	// Verify the number of inputs versus the expected number of inputs.
+	const UNearestNeighborModel* NearestNeighborModel = Cast<UNearestNeighborModel>(Model);
+	check(NearestNeighborModel);
+	const int32 NumFloatsPerBone = NearestNeighborModel->GetNumFloatsPerBone();
+	const int32 NumFloatsPerCurve = NearestNeighborModel->GetNumFloatsPerCurve();
+	const int64 NumDeformerAssetInputs = static_cast<int64>(Model->GetInputInfo()->CalcNumNeuralNetInputs(NumFloatsPerBone, NumFloatsPerCurve));
+	int64 NumNeuralNetInputs = -1;
+	if (NearestNeighborModel->DoesUseOptimizedNetwork())
+	{
+		const UNearestNeighborOptimizedNetwork* OptimizedNetwork = NearestNeighborModel->GetOptimizedNetwork();
+		if (OptimizedNetwork)
+		{
+			NumNeuralNetInputs = OptimizedNetwork->GetNumInputs();
+		}
+
+	}
+	else
+	{
+		UNeuralNetwork* NeuralNetwork = Model->GetNeuralNetwork();
+		if (NeuralNetwork && NeuralNetwork->IsLoaded())
+		{
+			NumNeuralNetInputs = NeuralNetwork->GetInputTensor().Num();
+		}
+	}
+
+	// Only check compatibility after the network is loaded. 
+	if (NumNeuralNetInputs >= 0 && NumNeuralNetInputs != NumDeformerAssetInputs) 
+	{
+		const FString InputErrorString = "The number of network inputs doesn't match the asset. Please retrain the asset."; 
+		ErrorText += InputErrorString + "\n";
+		if (LogIssues && Model->GetDeformerAsset())
+		{
+				UE_LOG(LogMLDeformer, Error, TEXT("Deformer '%s': %s"), *(Model->GetDeformerAsset()->GetName()), *InputErrorString);
+		}
+	}
+
+	return ErrorText;
+}
diff --git a/Engine/Plugins/Animation/MLDeformer/NearestNeighborModel/Source/NearestNeighborModel/Private/NearestNeighborOptimizedNetwork.cpp b/Engine/Plugins/Animation/MLDeformer/NearestNeighborModel/Source/NearestNeighborModel/Private/NearestNeighborOptimizedNetwork.cpp
new file mode 100644
--- /dev/null
+++ b/Engine/Plugins/Animation/MLDeformer/NearestNeighborModel/Source/NearestNeighborModel/Private/NearestNeighborOptimizedNetwork.cpp
@@ -0,0 +1,180 @@
+// Copyright Epic Games, Inc. All Rights Reserved.
+
+#include "NearestNeighborOptimizedNetwork.h"
+#include "NearestNeighborOptimizedNetwork.ispc.generated.h"
+
+//--------------------------------------------------------------------------
+// UNearestNeighborNetworkLayer
+//--------------------------------------------------------------------------
+void UNearestNeighborNetworkLayer::AddParameter(const TArray<float>& Values, const TArray<int32>& Shape)
+{
+	Parameters.Add({Values, Shape});
+}
+
+void UNearestNeighborNetworkLayer::Run(const float* RESTRICT InputBuffer, float* RESTRICT OutputBuffer) const
+{
+}
+
+void UNearestNeighborNetworkLayer_Gemm_Prelu::Run(const float* RESTRICT InputBuffer, float* RESTRICT OutputBuffer) const
+{
+	const float* Gemm_Weights = Parameters[0].Values.GetData();
+	const float* Gemm_Bias = Parameters[1].Values.GetData();
+	const float PRelu_Slope = Parameters[2].Values[0];
+	ispc::Gemm_PRelu(OutputBuffer, InputBuffer, Gemm_Weights, Gemm_Bias, PRelu_Slope, NumInputs, NumOutputs);
+}
+
+void UNearestNeighborNetworkLayer_Gemm::Run(const float* RESTRICT InputBuffer, float* RESTRICT OutputBuffer) const
+{
+	const float* Gemm_Weights = Parameters[0].Values.GetData();
+	const float* Gemm_Bias = Parameters[1].Values.GetData();
+	ispc::Gemm(OutputBuffer, InputBuffer, Gemm_Weights, Gemm_Bias, NumInputs, NumOutputs);
+}
+
+//--------------------------------------------------------------------------
+// UNearestNeighborOptimizedNetwork
+//--------------------------------------------------------------------------
+void UNearestNeighborOptimizedNetwork::Empty()
+{
+	Layers.Empty();
+}
+
+bool UNearestNeighborOptimizedNetwork::IsEmpty() const
+{
+	return Layers.IsEmpty();
+}
+
+int32 UNearestNeighborOptimizedNetwork::GetNumInputs() const
+{
+	return !Layers.IsEmpty() ? Layers[0]->NumInputs : 0;
+}
+
+int32 UNearestNeighborOptimizedNetwork::GetNumOutputs() const
+{
+	return !Layers.IsEmpty() ? Layers[Layers.Num()-1]->NumOutputs : 0;
+}
+
+bool UNearestNeighborOptimizedNetwork::Load(const FString& Filename)
+{
+	return true;
+}
+
+UNearestNeighborOptimizedNetworkInstance* UNearestNeighborOptimizedNetwork::CreateInstance()
+{
+	UNearestNeighborOptimizedNetworkInstance* Instance = NewObject<UNearestNeighborOptimizedNetworkInstance>(this);
+	Instance->Init(this);
+	return Instance;
+}
+
+const int32 UNearestNeighborOptimizedNetwork::GetNumLayers() const
+{
+	return Layers.Num();
+}
+
+UNearestNeighborNetworkLayer* UNearestNeighborOptimizedNetwork::GetLayer(int32 Index) const
+{
+	return Layers[Index].Get();
+}
+
+UNearestNeighborNetworkLayer* UNearestNeighborOptimizedNetwork::AddLayer(const int32 LayerType)
+{
+	UNearestNeighborNetworkLayer* Layer = nullptr;
+	switch ((ENearestNeighborNetworkLayerType)LayerType)
+	{
+		case ENearestNeighborNetworkLayerType::Gemm_Prelu:
+			Layer = NewObject<UNearestNeighborNetworkLayer_Gemm_Prelu>(this);
+			break;
+		case ENearestNeighborNetworkLayerType::Gemm:
+			Layer = NewObject<UNearestNeighborNetworkLayer_Gemm>(this);
+			break;
+		default:
+			break;
+	}
+	if (Layer)
+	{
+		Layers.Add(Layer);
+	}
+	return Layer;
+}
+
+//--------------------------------------------------------------------------
+// UNearestNeighborOptimizedNetworkInstance
+//--------------------------------------------------------------------------
+
+TArrayView<float> UNearestNeighborOptimizedNetworkInstance::GetInputs()
+{ 
+	return TArrayView<float>(Inputs.GetData(), Inputs.Num());
+}
+
+TArrayView<const float> UNearestNeighborOptimizedNetworkInstance::GetInputs() const
+{ 
+	return TArrayView<const float>(Inputs.GetData(), Inputs.Num());
+}
+
+TArrayView<float> UNearestNeighborOptimizedNetworkInstance::GetOutputs()
+{ 
+	return TArrayView<float>(Outputs.GetData(), Outputs.Num());
+}
+
+TArrayView<const float> UNearestNeighborOptimizedNetworkInstance::GetOutputs() const
+{ 
+	return TArrayView<const float>(Outputs.GetData(), Outputs.Num());
+}
+
+const UNearestNeighborOptimizedNetwork* UNearestNeighborOptimizedNetworkInstance::GetNeuralNetwork() const
+{ 
+	return Network.Get();
+}
+
+void UNearestNeighborOptimizedNetworkInstance::Init(UNearestNeighborOptimizedNetwork* InNeuralNetwork)
+{
+	Network = InNeuralNetwork;
+	Inputs.SetNumZeroed(Network->GetNumInputs());
+	Outputs.SetNumZeroed(Network->GetNumOutputs());
+
+	int32 MaxNumUnits = 0;
+	const int32 NumLayers = Network->GetNumLayers();
+	for (int32 LayerIndex = 0; LayerIndex < NumLayers; ++LayerIndex)
+	{
+		const UNearestNeighborNetworkLayer* CurLayer = Network->GetLayer(LayerIndex);
+		const int32 NumInputUnits = CurLayer->NumInputs;
+		const int32 NumOutputUnits = CurLayer->NumOutputs;
+		MaxNumUnits = FMath::Max3<int32>(NumInputUnits, NumOutputUnits, MaxNumUnits);
+	}
+
+	TempInputArray.SetNumZeroed(MaxNumUnits);
+	TempOutputArray.SetNumZeroed(MaxNumUnits);
+}
+
+
+void UNearestNeighborOptimizedNetworkInstance::Run()
+{	
+	TRACE_CPUPROFILER_EVENT_SCOPE(UNearestNeighborOptimizedNetwork::Run)
+
+	// Setup the buffer pointers.
+	FRunSettings RunSettings;
+	RunSettings.TempInputBuffer  = TempInputArray.GetData();
+	RunSettings.TempOutputBuffer = TempOutputArray.GetData();
+	RunSettings.InputBuffer		 = Inputs.GetData();
+	RunSettings.OutputBuffer	 = Outputs.GetData();
+
+	// Run the network.
+	const int32 NumLayers = Network->GetNumLayers();
+	for (int32 LayerIndex = 0; LayerIndex < NumLayers; ++LayerIndex)
+	{
+		const UNearestNeighborNetworkLayer* CurLayer = Network->GetLayer(LayerIndex);
+		if (LayerIndex == 0)
+		{
+			const float* const RESTRICT NetworkInputs = RunSettings.InputBuffer;
+			CurLayer->Run(RunSettings.InputBuffer, RunSettings.TempInputBuffer);
+		}
+		else if (LayerIndex == NumLayers - 1)
+		{
+			CurLayer->Run(RunSettings.TempInputBuffer, RunSettings.OutputBuffer);
+		}
+		else
+		{
+			CurLayer->Run(RunSettings.TempInputBuffer, RunSettings.TempOutputBuffer);
+			Swap(RunSettings.TempInputBuffer, RunSettings.TempOutputBuffer);
+		}
+	}
+}
diff --git a/Engine/Plugins/Animation/MLDeformer/NearestNeighborModel/Source/NearestNeighborModel/Private/NearestNeighborOptimizedNetwork.h b/Engine/Plugins/Animation/MLDeformer/NearestNeighborModel/Source/NearestNeighborModel/Private/NearestNeighborOptimizedNetwork.h
new file mode 100644
--- /dev/null
+++ b/Engine/Plugins/Animation/MLDeformer/NearestNeighborModel/Source/NearestNeighborModel/Private/NearestNeighborOptimizedNetwork.h
@@ -0,0 +1,227 @@
+// Copyright Epic Games, Inc. All Rights Reserved.
+#pragma once
+
+#include "CoreTypes.h"
+#include "UObject/Object.h"
+#include "NearestNeighborOptimizedNetwork.generated.h"
+
+class UNearestNeighborOptimizedNetworkInstance;
+
+USTRUCT()
+struct FNearestNeighborNetworkParameter
+{
+	GENERATED_BODY()
+
+	UPROPERTY()
+	TArray<float> Values;
+
+	UPROPERTY()
+	TArray<int32> Shape;
+};
+
+UENUM(BlueprintType)
+enum class ENearestNeighborNetworkLayerType : uint8
+{
+	None,
+	Gemm_Prelu,
+	Gemm,
+};
+
+
+/** A general network layer that contains a list of parameters. The Run() method should be implemented by child classes */ 
+UCLASS()
+class UNearestNeighborNetworkLayer
+	: public UObject
+{
+	GENERATED_BODY()
+
+public:
+	/** The weight matrix number of inputs (rows). */
+	UPROPERTY(EditAnywhere, BlueprintReadWrite, Category = "MLDeformer")
+	int32 NumInputs = 0;
+
+	/** The weight matrix number of outputs (columns). */
+	UPROPERTY(EditAnywhere, BlueprintReadWrite, Category = "MLDeformer")
+	int32 NumOutputs = 0;
+
+	// /** The parameters of the layer */
+	UPROPERTY()
+	TArray<FNearestNeighborNetworkParameter> Parameters;
+
+	UFUNCTION(BlueprintCallable, Category = "MLDeformer")
+	void AddParameter(const TArray<float>& Values, const TArray<int32>& Shape);
+
+	virtual void Run(const float* RESTRICT InputBuffer, float* RESTRICT OutputBuffer) const;
+};
+
+/** Gemm: GEneral Matrix Multiplication: Wx + b */
+UCLASS()
+class UNearestNeighborNetworkLayer_Gemm_Prelu
+	: public UNearestNeighborNetworkLayer
+{
+	GENERATED_BODY()
+
+public:
+	virtual void Run(const float* RESTRICT InputBuffer, float* RESTRICT OutputBuffer) const override;
+};
+
+UCLASS()
+class UNearestNeighborNetworkLayer_Gemm
+	: public UNearestNeighborNetworkLayer
+{
+	GENERATED_BODY()
+
+public:
+	virtual void Run(const float* RESTRICT InputBuffer, float* RESTRICT OutputBuffer) const override;
+};
+
+/**
+ * The specialized neural network for the MLDeformerModel.
+ * This class is used to do inference at runtime at a higher performance than using UNeuralNetwork.
+ */
+UCLASS()
+class UNearestNeighborOptimizedNetwork
+	: public UObject
+{
+	GENERATED_BODY()
+
+public:
+	/** Clear the network, getting rid of all layers. */
+	UFUNCTION(BlueprintCallable, Category = "MLDeformer")
+	virtual void Empty();
+
+	/** 
+	 * Check if the network is empty or not.
+	 * If it is empty, it means it hasn't been loaded, and cannot do anything.
+	 * Empty means it has no inputs, outputs or weights.
+	 * @return Returns true if empty, otherwise false is returned.
+	 */
+	virtual bool IsEmpty() const;
+
+	/**
+	 * Load the network from a file on disk.
+	 * When loading fails, the network will be emptied.
+	 * @param Filename The file path to load from.
+	 * @return Returns true when successfully loaded this model, otherwise false is returned.
+	 */
+	virtual bool Load(const FString& Filename);
+
+	/**
+	 * Create an instance of this neural network.
+	 * @return A pointer to the neural network instance.
+	 */
+	virtual UNearestNeighborOptimizedNetworkInstance* CreateInstance();
+
+	/**
+	 * Get the number of inputs, which is the number of floats the network takes as input.
+	 * @result The number of input floats to the network.
+	 */
+	virtual int32 GetNumInputs() const;
+
+	/**
+	 * Get the number of outputs, which is the number of floats the network will output.
+	 * @return The number of floats that the network outputs.
+	 */
+	virtual int32 GetNumOutputs() const;
+
+
+	/**
+	 * Get the number of network layers.
+	 * @return The number of network layer.
+	 */
+	const int32 GetNumLayers() const;
+
+	/**
+	 * Get a given network layer.
+	 * @return A pointer to the layer, which will contain the parameters.
+	 */
+	UNearestNeighborNetworkLayer* GetLayer(int32 Index) const;
+
+
+	/**
+	 * Add a network layer.
+	 * @return A pointer to the layer, which will contain the parameters.
+	 */
+	UFUNCTION(BlueprintCallable, Category = "MLDeformer")
+	UNearestNeighborNetworkLayer* AddLayer(const int32 LayerType);
+
+protected:
+	/** The network layers */
+	UPROPERTY()
+	TArray<TObjectPtr<UNearestNeighborNetworkLayer>> Layers;
+};
+
+/** 
+ * An instance of a UNearestNeighborOptimizedNetwork.
+ * The instance holds its own input and output buffers and only reads from the network object it was instanced from.
+ * This allows it to be multithreaded.
+ */
+UCLASS()
+class UNearestNeighborOptimizedNetworkInstance
+	: public UObject
+{
+	friend class UNearestNeighborOptimizedNetwork;
+
+	GENERATED_BODY()
+
+public:
+	/**
+	 * Get the network input buffer.
+	 * @return The array view of floats that represents the network inputs.
+	 */
+	TArrayView<float> GetInputs();
+	TArrayView<const float> GetInputs() const;
+
+	/**
+	 * Get the network output buffer.
+	 * @return The array view of floats that represents the network outputs.
+	 */
+	TArrayView<float> GetOutputs();
+	TArrayView<const float> GetOutputs() const;
+
+	/**
+	 * Get the neural network this is an instance of.
+	 * @return A reference to the neural network.
+	 */
+	const UNearestNeighborOptimizedNetwork* GetNeuralNetwork() const;
+
+	/**
+	 * Run the neural network, performing inference.
+	 * This will update the values in the output buffer that you can get with the GetOutputs method.
+	 * This also assumes you have set all the right input values already.
+	 */
+	virtual void Run();
+
+protected:
+	struct FRunSettings
+	{
+		float* TempInputBuffer;
+		float* TempOutputBuffer;
+		const float* InputBuffer;
+		float* OutputBuffer;
+	};
+
+	/**
+ 	 * Initialize this instance using a given network object.
+	 * This is automatically called by the UNearestNeighborOptimizedNetwork::CreateInstance() method.
+	 * @param InNeuralNetwork The network that this is an instance of.
+	 */
+	void Init(UNearestNeighborOptimizedNetwork* InNeuralNetwork);
+
+protected:
+	/** The input values. */
+	TArray<float> Inputs;
+
+	/** The output values. */
+	TArray<float> Outputs;
+
+	/** A pre-allocated temp buffer for inputs. */
+	TArray<float> TempInputArray;
+
+	/** A pre-allocated temp buffer for outputs. */
+	TArray<float> TempOutputArray;
+
+	/** The neural network this is an instance of. */
+	UPROPERTY(Transient)
+	TObjectPtr<UNearestNeighborOptimizedNetwork> Network;
+};
diff --git a/Engine/Plugins/Animation/MLDeformer/NearestNeighborModel/Source/NearestNeighborModel/Private/NearestNeighborOptimizedNetwork.ispc b/Engine/Plugins/Animation/MLDeformer/NearestNeighborModel/Source/NearestNeighborModel/Private/NearestNeighborOptimizedNetwork.ispc
new file mode 100644
--- /dev/null
+++ b/Engine/Plugins/Animation/MLDeformer/NearestNeighborModel/Source/NearestNeighborModel/Private/NearestNeighborOptimizedNetwork.ispc
@@ -0,0 +1,59 @@
+// Copyright Epic Games, Inc. All Rights Reserved.
+
+export void Gemm_PRelu(
+	uniform float Outputs[],
+	const uniform float Inputs[],
+	const uniform float Gemm_Weights[],
+	const uniform float Gemm_Biases[],
+	const uniform float PRelu_Alpha,
+	const uniform int NumInputs,
+	const uniform int NumOutputs)
+{
+	foreach (OutputIndex = 0 ... NumOutputs)
+	{
+		Outputs[OutputIndex] = Gemm_Biases[OutputIndex];
+	}
+
+	for (uniform int InputIndex = 0; InputIndex < NumInputs; ++InputIndex)
+	{
+		uniform float X = Inputs[InputIndex];
+		foreach (OutputIndex = 0 ... NumOutputs)
+		{
+			Outputs[OutputIndex] += X * Gemm_Weights[InputIndex * NumOutputs + OutputIndex];
+		}
+	}
+
+	foreach (OutputIndex = 0 ... NumOutputs)
+	{
+		const float Output = Outputs[OutputIndex];
+		Outputs[OutputIndex] = select(Output > 0.0f, Output, Output * PRelu_Alpha);
+	}
+}
+
+
+export void Gemm(
+	uniform float Outputs[],
+	const uniform float Inputs[],
+	const uniform float Gemm_Weights[],
+	const uniform float Gemm_Biases[],
+	const uniform int NumInputs,
+	const uniform int NumOutputs)
+{
+	foreach (OutputIndex = 0 ... NumOutputs)
+	{
+		Outputs[OutputIndex] = Gemm_Biases[OutputIndex];
+	}
+
+	for (uniform int InputIndex = 0; InputIndex < NumInputs; ++InputIndex)
+	{
+		uniform float X = Inputs[InputIndex];
+		foreach (OutputIndex = 0 ... NumOutputs)
+		{
+			Outputs[OutputIndex] += X * Gemm_Weights[InputIndex * NumOutputs + OutputIndex];
+		}
+	}
+}
+
+
+
+
diff --git a/Engine/Plugins/Animation/MLDeformer/NearestNeighborModel/Source/NearestNeighborModel/Private/NearestNeighborOptimizedNetworkLoader.cpp b/Engine/Plugins/Animation/MLDeformer/NearestNeighborModel/Source/NearestNeighborModel/Private/NearestNeighborOptimizedNetworkLoader.cpp
new file mode 100644
--- /dev/null
+++ b/Engine/Plugins/Animation/MLDeformer/NearestNeighborModel/Source/NearestNeighborModel/Private/NearestNeighborOptimizedNetworkLoader.cpp
@@ -0,0 +1,14 @@
+// Copyright Epic Games, Inc. All Rights Reserved.
+
+#include "NearestNeighborOptimizedNetworkLoader.h"
+#include "NearestNeighborOptimizedNetwork.h"
+
+void UNearestNeighborOptimizedNetworkLoader::SetOptimizedNetwork(UNearestNeighborOptimizedNetwork* InNetwork)
+{
+	Network = InNetwork;
+}
+
+UNearestNeighborOptimizedNetwork* UNearestNeighborOptimizedNetworkLoader::GetOptimizedNetwork() const
+{
+	return Network;
+}
diff --git a/Engine/Plugins/Animation/MLDeformer/NearestNeighborModel/Source/NearestNeighborModel/Public/NearestNeighborModel.h b/Engine/Plugins/Animation/MLDeformer/NearestNeighborModel/Source/NearestNeighborModel/Public/NearestNeighborModel.h
--- a/Engine/Plugins/Animation/MLDeformer/NearestNeighborModel/Source/NearestNeighborModel/Public/NearestNeighborModel.h
+++ b/Engine/Plugins/Animation/MLDeformer/NearestNeighborModel/Source/NearestNeighborModel/Public/NearestNeighborModel.h
@@ -18,6 +18,7 @@ class UNeuralNetwork;
 class UMLDeformerAsset;
 class USkeleton;
 class IPropertyHandle;
+class UNearestNeighborOptimizedNetwork;
 
 NEARESTNEIGHBORMODEL_API DECLARE_LOG_CATEGORY_EXTERN(LogNearestNeighborModel, Log, All);
 
@@ -157,6 +158,7 @@ public:
 	virtual void PostLoad() override;
 
 	// UMLDeformerModel overrides.
+	virtual void Serialize(FArchive& Archive) override;
 	virtual bool DoesSupportQualityLevels() const override { return false; }
 	virtual UMLDeformerModelInstance* CreateModelInstance(UMLDeformerComponent* Component) override;
 	virtual UMLDeformerInputInfo* CreateInputInfo() override;
@@ -233,6 +235,15 @@ public:
 
 	void ResetMorphBuffers();
 
+	virtual int32 GetNumFloatsPerBone() const { return NearestNeighborNumFloatsPerBone; }
+	virtual int32 GetNumFloatsPerCurve() const { return NearestNeighborNumFloatsPerCurve; }
+
+	bool DoesUseOptimizedNetwork() const;
+	UNearestNeighborOptimizedNetwork* GetOptimizedNetwork() const { return OptimizedNetwork.Get(); }
+	int32 GetOptimizedNetworkNumOutputs() const;
+	void SetOptimizedNetwork(UNearestNeighborOptimizedNetwork* InOptimizedNetwork);
+	bool LoadOptimizedNetwork(const FString& OnnxPath);
+
 #if WITH_EDITORONLY_DATA
 	TObjectPtr<UAnimSequence> GetNearestNeighborSkeletons(int32 PartId);
 	const TObjectPtr<UAnimSequence> GetNearestNeighborSkeletons(int32 PartId) const;
@@ -373,6 +384,9 @@ protected:
 	UPROPERTY(EditAnywhere, BlueprintReadWrite, Category = "File Cache")
 	bool bRecomputePCA = true;
 #endif
+
+	static constexpr int32 NearestNeighborNumFloatsPerBone = 3;
+	static constexpr int32 NearestNeighborNumFloatsPerCurve = 0;
 	
 public:
 #if WITH_EDITORONLY_DATA
@@ -406,4 +420,7 @@ public:
 
 	UPROPERTY(EditAnywhere, BlueprintReadWrite, Category = "Nearest Neighbors", META = (ClampMin = "0", ClampMax = "1"))
 	float NearestNeighborOffsetWeight = 1.0f;
+
+	UPROPERTY()
+	TObjectPtr<UNearestNeighborOptimizedNetwork> OptimizedNetwork = nullptr;
 };
diff --git a/Engine/Plugins/Animation/MLDeformer/NearestNeighborModel/Source/NearestNeighborModel/Public/NearestNeighborModelInstance.h b/Engine/Plugins/Animation/MLDeformer/NearestNeighborModel/Source/NearestNeighborModel/Public/NearestNeighborModelInstance.h
--- a/Engine/Plugins/Animation/MLDeformer/NearestNeighborModel/Source/NearestNeighborModel/Public/NearestNeighborModelInstance.h
+++ b/Engine/Plugins/Animation/MLDeformer/NearestNeighborModel/Source/NearestNeighborModel/Public/NearestNeighborModelInstance.h
@@ -5,6 +5,8 @@
 #include "NeuralTensor.h"
 #include "NearestNeighborModelInstance.generated.h"
 
+class UNearestNeighborOptimizedNetworkInstance;
+
 UCLASS()
 class NEARESTNEIGHBORMODEL_API UNearestNeighborModelInstance
     : public UMLDeformerMorphModelInstance
@@ -16,6 +18,7 @@ public:
     virtual void Init(USkeletalMeshComponent* SkelMeshComponent) override;
     virtual void Execute(float ModelWeight) override;
     virtual bool SetupInputs() override;
+    virtual FString CheckCompatibility(USkeletalMeshComponent* InSkelMeshComponent, bool LogIssues=false) override;
 
 #if WITH_EDITORONLY_DATA
     const TArray<uint32>& GetNearestNeighborIds() const { return NearestNeighborIds; }
@@ -24,6 +27,9 @@ public:
 #endif
 
     void InitPreviousWeights();
+    void InitOptimizedNetworkInstance();
+    void GetInputDataPointer(float*& OutInputData, int32& OutNumInputFloats) const;
+    void GetOutputDataPointer(float*& OutOutputData, int32& OutNumOutputFloats) const;
 
 protected:
     virtual int64 SetBoneTransforms(float* OutputBuffer, int64 OutputBufferSize, int64 StartIndex) override;
@@ -31,7 +37,7 @@ protected:
 
 private:
     void RunNearestNeighborModel(float ModelWeight);
-    int32 FindNearestNeighbor(const FNeuralTensor& PCACoeffTensor, int32 PartId);
+    int32 FindNearestNeighbor(const float* PCAData, int32 PartId);
     void UpdateWeight(TArray<float>& MorphWeights, int32 Index, float W);
 
 #if WITH_EDITORONLY_DATA
@@ -39,4 +45,7 @@ private:
 #endif
 
     TArray<float> PreviousWeights;
+
+    UPROPERTY(Transient)
+    TObjectPtr<UNearestNeighborOptimizedNetworkInstance> OptimizedNetworkInstance = nullptr;
 };
diff --git a/Engine/Plugins/Animation/MLDeformer/NearestNeighborModel/Source/NearestNeighborModel/Public/NearestNeighborOptimizedNetworkLoader.h b/Engine/Plugins/Animation/MLDeformer/NearestNeighborModel/Source/NearestNeighborModel/Public/NearestNeighborOptimizedNetworkLoader.h
new file mode 100644
--- /dev/null
+++ b/Engine/Plugins/Animation/MLDeformer/NearestNeighborModel/Source/NearestNeighborModel/Public/NearestNeighborOptimizedNetworkLoader.h
@@ -0,0 +1,30 @@
+// Copyright Epic Games, Inc. All Rights Reserved.
+
+#pragma once
+
+#include "CoreMinimal.h"
+#include "NearestNeighborOptimizedNetworkLoader.generated.h"
+
+class UNearestNeighborOptimizedNetwork;
+
+/** Helper class to load the optimized network from disk.
+ *  LoadOptimizedNetwork is implemented in python.
+ */
+UCLASS()
+class NEARESTNEIGHBORMODEL_API UNearestNeighborOptimizedNetworkLoader
+	: public UObject
+{
+	GENERATED_BODY()
+
+public:
+	UFUNCTION(BlueprintImplementableEvent, Category = "Nearest Neighbor Model")
+	bool LoadOptimizedNetwork(const FString& OnnxPath);
+
+	void SetOptimizedNetwork(UNearestNeighborOptimizedNetwork* InNetwork);
+
+	UFUNCTION(BlueprintPure, Category = "Nearest Neighbor Model")
+	UNearestNeighborOptimizedNetwork* GetOptimizedNetwork() const;
+
+protected:
+	TObjectPtr<UNearestNeighborOptimizedNetwork> Network = nullptr;
+};
\ No newline at end of file
diff --git a/Engine/Plugins/Animation/MLDeformer/NearestNeighborModel/Source/NearestNeighborModelEditor/Private/NearestNeighborEditorModel.cpp b/Engine/Plugins/Animation/MLDeformer/NearestNeighborModel/Source/NearestNeighborModelEditor/Private/NearestNeighborEditorModel.cpp
--- a/Engine/Plugins/Animation/MLDeformer/NearestNeighborModel/Source/NearestNeighborModelEditor/Private/NearestNeighborEditorModel.cpp
+++ b/Engine/Plugins/Animation/MLDeformer/NearestNeighborModel/Source/NearestNeighborModelEditor/Private/NearestNeighborEditorModel.cpp
@@ -90,6 +90,44 @@ namespace UE::NearestNeighborModel
 		}
 	}
 
+	bool FNearestNeighborEditorModel::LoadTrainedNetwork() const
+	{
+		UNearestNeighborModel* NearestNeighborModel = GetNearestNeighborModel();
+		if (NearestNeighborModel)
+		{
+			const FString OnnxFile = GetTrainedNetworkOnnxFile();
+			UNeuralNetwork* Network = LoadNeuralNetworkFromOnnx(OnnxFile);
+			if (Network)
+			{
+				if (NearestNeighborModel->DoesUseOptimizedNetwork())
+				{
+					Model->SetNeuralNetwork(Network);
+					return NearestNeighborModel->LoadOptimizedNetwork(OnnxFile);
+				}
+				else
+				{
+					Model->SetNeuralNetwork(Network);
+					return true;
+				}
+			}
+			else
+			{
+				return false;
+			}
+		}
+		return false;
+	}
+
+	bool FNearestNeighborEditorModel::IsTrained() const
+	{
+		const UNearestNeighborModel* NearestNeighborModel = GetNearestNeighborModel();
+		if (NearestNeighborModel)
+		{
+			return NearestNeighborModel->DoesUseOptimizedNetwork() ? NearestNeighborModel->GetOptimizedNetwork() != nullptr : NearestNeighborModel->GetNeuralNetwork() != nullptr;
+		}
+		return false;
+	}
+
 	FString FNearestNeighborEditorModel::GetTrainedNetworkOnnxFile() const
 	{
 		const UNearestNeighborModel* NearestNeighborModel = GetNearestNeighborModel();
@@ -144,14 +182,10 @@ namespace UE::NearestNeighborModel
 					NearestNeighborActor->SetMeshOffsetFactor(Offset);
 				}
 			}
-			const UMLDeformerComponent* MLDeformerComponent = GetTestMLDeformerComponent();
-			if (MLDeformerComponent)
+			const UNearestNeighborModelInstance* ModelInstance = static_cast<UNearestNeighborModelInstance*>(GetTestMLDeformerModelInstance());
+			if (ModelInstance)
 			{
-				const UNearestNeighborModelInstance* ModelInstance = static_cast<UNearestNeighborModelInstance*>(MLDeformerComponent->GetModelInstance());
-				if (ModelInstance)
-				{
-					NNViz->SetNearestNeighborIds(ModelInstance->GetNearestNeighborIds());
-				}
+				NNViz->SetNearestNeighborIds(ModelInstance->GetNearestNeighborIds());
 			}
 		}
 	}
@@ -371,7 +405,7 @@ namespace UE::NearestNeighborModel
 			return EUpdateResult::SUCCESS;
 		}
 
-		UNeuralNetwork *NeuralNetwork = NearestNeighborModel->GetNeuralNetwork();
+		const UNeuralNetwork *NeuralNetwork = NearestNeighborModel->GetNeuralNetwork();
 		if(NeuralNetwork && NeuralNetwork->IsLoaded())
 		{
 			const FString SavePath = GetTrainedNetworkOnnxFile();
@@ -467,7 +501,7 @@ namespace UE::NearestNeighborModel
 		UNearestNeighborModel *NearestNeighborModel = static_cast<UNearestNeighborModel*>(Model);
 		check(NearestNeighborModel != nullptr);
 
-		UNeuralNetwork *NeuralNetwork = NearestNeighborModel->GetNeuralNetwork();
+		const UNeuralNetwork *NeuralNetwork = NearestNeighborModel->GetNeuralNetwork();
 		if(NeuralNetwork && NeuralNetwork->IsLoaded())
 		{
 			const FString SavePath = GetTrainedNetworkOnnxFile();
@@ -643,6 +677,16 @@ namespace UE::NearestNeighborModel
 		MorphTargetUpdateResult = EUpdateResult::SUCCESS;
 
 		UNearestNeighborModel* NearestNeighborModel = GetNearestNeighborModel();
+		// Convert NNI network to  optimized network
+		const UNeuralNetwork *NeuralNetwork = NearestNeighborModel->GetNeuralNetwork();
+		if (NearestNeighborModel->DoesUseOptimizedNetwork() && NeuralNetwork != nullptr && NearestNeighborModel->GetOptimizedNetwork() == nullptr)
+		{
+			const FString SavePath = GetTrainedNetworkOnnxFile();
+			UE_LOG(LogNearestNeighborModel, Display, TEXT("Saving to %s"), *SavePath);
+			NeuralNetwork->Save(SavePath);
+			LoadTrainedNetwork();
+		}
+
 		if (!NearestNeighborModel->IsClothPartDataValid())
 		{
 			MorphTargetUpdateResult |= NearestNeighborModel->UpdateClothPartData();
@@ -797,31 +841,18 @@ namespace UE::NearestNeighborModel
 		}
 	}
 
-	UMLDeformerComponent* FNearestNeighborEditorModel::GetTestMLDeformerComponent()
+	UMLDeformerComponent* FNearestNeighborEditorModel::GetTestMLDeformerComponent() const
 	{
-		UMLDeformerComponent* MLDeformerComponent = nullptr;
-		FMLDeformerEditorActor* TestActor = FindEditorActor(ActorID_Test_MLDeformed);
-		if (TestActor)
-		{
-			AActor* Actor = TestActor->GetActor();
-			if (Actor)
-			{
-				MLDeformerComponent = Actor->FindComponentByClass<UMLDeformerComponent>();
-			}
-		}
-		return MLDeformerComponent;
+		return FindMLDeformerComponent(ActorID_Test_MLDeformed);
 	}
 
-	void FNearestNeighborEditorModel::InitTestMLDeformerPreviousWeights()
+
+	void FNearestNeighborEditorModel::InitTestMLDeformerPreviousWeights() 
 	{
-		UMLDeformerComponent* MLDeformerComponent = GetTestMLDeformerComponent();
-		if (MLDeformerComponent)
+		UNearestNeighborModelInstance* ModelInstance = static_cast<UNearestNeighborModelInstance*>(GetTestMLDeformerModelInstance());
+		if (ModelInstance)
 		{
-			UNearestNeighborModelInstance* ModelInstance = static_cast<UNearestNeighborModelInstance*>(MLDeformerComponent->GetModelInstance());
-			if (ModelInstance)
-			{
-				ModelInstance->InitPreviousWeights();
-			}
+			ModelInstance->InitPreviousWeights();
 		}
 	}
 
@@ -829,14 +860,9 @@ namespace UE::NearestNeighborModel
 	{
 		const UNearestNeighborModel *NearestNeighborModel = static_cast<UNearestNeighborModel*>(Model);
 
-		const UNeuralNetwork *NeuralNetwork = NearestNeighborModel->GetNeuralNetwork();
-		const UMLDeformerComponent* MLDeformerComponent = GetTestMLDeformerComponent();
-		if(NeuralNetwork && NeuralNetwork->IsLoaded() && MLDeformerComponent && MLDeformerComponent->GetModelInstance())
+		if (NearestNeighborModel->DoesUseOptimizedNetwork())
 		{
-			const int32 NeuralNetworkInferenceHandle = MLDeformerComponent->GetModelInstance()->GetNeuralNetworkInferenceHandle();
-			const FNeuralTensor& OutputTensor = NeuralNetwork->GetOutputTensorForContext(NeuralNetworkInferenceHandle);
-			const int32 NumNetworkWeights = OutputTensor.Num();
-
+			const int32 NumNetworkWeights = NearestNeighborModel->GetOptimizedNetworkNumOutputs();
 			const int32 NumPCACoeffs = NearestNeighborModel->GetTotalNumPCACoeffs();
 			if (NumNetworkWeights != NumPCACoeffs)
 			{
@@ -844,6 +870,24 @@ namespace UE::NearestNeighborModel
 				return EUpdateResult::WARNING;
 			}
 		}
+		else
+		{
+			const UNeuralNetwork *NeuralNetwork = NearestNeighborModel->GetNeuralNetwork();
+			const UMLDeformerComponent* MLDeformerComponent = GetTestMLDeformerComponent();
+			if(NeuralNetwork && NeuralNetwork->IsLoaded() && MLDeformerComponent && MLDeformerComponent->GetModelInstance())
+			{
+				const int32 NeuralNetworkInferenceHandle = MLDeformerComponent->GetModelInstance()->GetNeuralNetworkInferenceHandle();
+				const FNeuralTensor& OutputTensor = NeuralNetwork->GetOutputTensorForContext(NeuralNetworkInferenceHandle);
+				const int32 NumNetworkWeights = OutputTensor.Num();
+				const int32 NumPCACoeffs = NearestNeighborModel->GetTotalNumPCACoeffs();
+				if (NumNetworkWeights != NumPCACoeffs)
+				{
+					UE_LOG(LogNearestNeighborModel, Warning, TEXT("Network output dimension %d is not equal to number of morph targets %d. Network needs to be re-trained and no deformation will be applied."), NumNetworkWeights, NumPCACoeffs);
+					return EUpdateResult::WARNING;
+				}
+			}
+		}
+
 		return EUpdateResult::SUCCESS;
 	}
 
@@ -853,6 +897,17 @@ namespace UE::NearestNeighborModel
 		const UNeuralNetwork *NeuralNetwork = NearestNeighborModel->GetNeuralNetwork();
 		return NeuralNetwork && NeuralNetwork->IsLoaded();
 	}
+
+	UMLDeformerModelInstance* FNearestNeighborEditorModel::GetTestMLDeformerModelInstance() const
+	{
+		UMLDeformerModelInstance* ModelInstance = nullptr;
+		const UMLDeformerComponent* MLDeformerComponent = FindMLDeformerComponent(ActorID_Test_MLDeformed);
+		if (MLDeformerComponent)
+		{
+			ModelInstance = MLDeformerComponent->GetModelInstance();
+		}
+		return ModelInstance;
+	}
 }	// namespace UE::NearestNeighborModel
 
 #undef LOCTEXT_NAMESPACE
diff --git a/Engine/Plugins/Animation/MLDeformer/NearestNeighborModel/Source/NearestNeighborModelEditor/Public/NearestNeighborEditorModel.h b/Engine/Plugins/Animation/MLDeformer/NearestNeighborModel/Source/NearestNeighborModelEditor/Public/NearestNeighborEditorModel.h
--- a/Engine/Plugins/Animation/MLDeformer/NearestNeighborModel/Source/NearestNeighborModelEditor/Public/NearestNeighborEditorModel.h
+++ b/Engine/Plugins/Animation/MLDeformer/NearestNeighborModel/Source/NearestNeighborModelEditor/Public/NearestNeighborEditorModel.h
@@ -10,7 +10,9 @@ class UMLDeformerAsset;
 class USkeletalMesh;
 class UGeometryCache;
 class UMLDeformerComponent;
+class UMLDeformerOptimizedNetwork;
 class UNearestNeighborModel;
+class UNearestNeighborModelInstance;
 class UNearestNeighborModelVizSettings;
 
 namespace UE::NearestNeighborModel
@@ -40,6 +42,8 @@ namespace UE::NearestNeighborModel
 
 		virtual void InitInputInfo(UMLDeformerInputInfo* InputInfo) override;
 		virtual ETrainingResult Train() override;
+		virtual bool IsTrained() const override;
+		virtual bool LoadTrainedNetwork() const override;
 		virtual void OnPropertyChanged(FPropertyChangedEvent& PropertyChangedEvent) override;
 		virtual void OnPostTraining(ETrainingResult TrainingResult, bool bUsePartiallyTrainedWhenAborted) override;
 		virtual FString GetTrainedNetworkOnnxFile() const override;
@@ -72,12 +76,12 @@ namespace UE::NearestNeighborModel
 
 		void OnMorphTargetUpdate();
 		uint8 GetMorphTargetUpdateResult() { return MorphTargetUpdateResult; }
-		UMLDeformerComponent* GetTestMLDeformerComponent();
+		UMLDeformerComponent* GetTestMLDeformerComponent() const;
+		UMLDeformerModelInstance* GetTestMLDeformerModelInstance() const;
 		void InitTestMLDeformerPreviousWeights();
 		uint8 WarnIfNetworkInvalid();
 		bool IsNeuralNetworkLoaded();
 
-
 	protected:
 		virtual void CreateNearestNeighborActors(UWorld* World, int32 StartIndex = 0);
 
