NearestNeighborModel: fix ISPC macro for MAC

diff --git a/Engine/Plugins/Animation/MLDeformer/NearestNeighborModel/Source/NearestNeighborModel/Private/NearestNeighborModel.cpp b/Engine/Plugins/Animation/MLDeformer/NearestNeighborModel/Source/NearestNeighborModel/Private/NearestNeighborModel.cpp
--- a/Engine/Plugins/Animation/MLDeformer/NearestNeighborModel/Source/NearestNeighborModel/Private/NearestNeighborModel.cpp
+++ b/Engine/Plugins/Animation/MLDeformer/NearestNeighborModel/Source/NearestNeighborModel/Private/NearestNeighborModel.cpp
@@ -462,7 +462,11 @@ void UNearestNeighborModel::SetOptimizedNetwork(UNearestNeighborOptimizedNetwork
 
 bool UNearestNeighborModel::DoesUseOptimizedNetwork() const
 {
+#if NEARESTNEIGHBORMODEL_USE_ISPC
 	return UE::NearestNeighborModel::bNearestNeighborModelUseOptimizedNetwork;
+#else
+	return false;
+#endif
 }
 
 /**
diff --git a/Engine/Plugins/Animation/MLDeformer/NearestNeighborModel/Source/NearestNeighborModel/Private/NearestNeighborOptimizedNetwork.cpp b/Engine/Plugins/Animation/MLDeformer/NearestNeighborModel/Source/NearestNeighborModel/Private/NearestNeighborOptimizedNetwork.cpp
--- a/Engine/Plugins/Animation/MLDeformer/NearestNeighborModel/Source/NearestNeighborModel/Private/NearestNeighborOptimizedNetwork.cpp
+++ b/Engine/Plugins/Animation/MLDeformer/NearestNeighborModel/Source/NearestNeighborModel/Private/NearestNeighborOptimizedNetwork.cpp
@@ -1,8 +1,9 @@
 // Copyright Epic Games, Inc. All Rights Reserved.
 
 #include "NearestNeighborOptimizedNetwork.h"
+#if NEARESTNEIGHBORMODEL_USE_ISPC
 #include "NearestNeighborOptimizedNetwork.ispc.generated.h"
-
+#endif
 //--------------------------------------------------------------------------
 // UNearestNeighborNetworkLayer
 //--------------------------------------------------------------------------
@@ -20,14 +21,18 @@ void UNearestNeighborNetworkLayer_Gemm_Prelu::Run(const float* RESTRICT InputBuf
 	const float* Gemm_Weights = Parameters[0].Values.GetData();
 	const float* Gemm_Bias = Parameters[1].Values.GetData();
 	const float PRelu_Slope = Parameters[2].Values[0];
+#if NEARESTNEIGHBORMODEL_USE_ISPC
 	ispc::Gemm_PRelu(OutputBuffer, InputBuffer, Gemm_Weights, Gemm_Bias, PRelu_Slope, NumInputs, NumOutputs);
+#endif
 }
 
 void UNearestNeighborNetworkLayer_Gemm::Run(const float* RESTRICT InputBuffer, float* RESTRICT OutputBuffer) const
 {
 	const float* Gemm_Weights = Parameters[0].Values.GetData();
 	const float* Gemm_Bias = Parameters[1].Values.GetData();
+#if NEARESTNEIGHBORMODEL_USE_ISPC
 	ispc::Gemm(OutputBuffer, InputBuffer, Gemm_Weights, Gemm_Bias, NumInputs, NumOutputs);
+#endif
 }
 
 //--------------------------------------------------------------------------
diff --git a/Engine/Plugins/Animation/MLDeformer/NearestNeighborModel/Source/NearestNeighborModel/Private/NearestNeighborOptimizedNetwork.h b/Engine/Plugins/Animation/MLDeformer/NearestNeighborModel/Source/NearestNeighborModel/Private/NearestNeighborOptimizedNetwork.h
--- a/Engine/Plugins/Animation/MLDeformer/NearestNeighborModel/Source/NearestNeighborModel/Private/NearestNeighborOptimizedNetwork.h
+++ b/Engine/Plugins/Animation/MLDeformer/NearestNeighborModel/Source/NearestNeighborModel/Private/NearestNeighborOptimizedNetwork.h
@@ -7,6 +7,11 @@
 
 class UNearestNeighborOptimizedNetworkInstance;
 
+#define NEARESTNEIGHBORMODEL_USE_ISPC INTEL_ISPC
+#if !defined(NEARESTNEIGHBORMODEL_USE_ISPC)
+	#define NEARESTNEIGHBORMODEL_USE_ISPC 0
+#endif
+
 USTRUCT()
 struct FNearestNeighborNetworkParameter
 {
