NearestNeigbhorModel: fix NNI to ISPC conversion

diff --git a/Engine/Plugins/Animation/MLDeformer/NearestNeighborModel/Source/NearestNeighborModelEditor/Private/NearestNeighborEditorModel.cpp b/Engine/Plugins/Animation/MLDeformer/NearestNeighborModel/Source/NearestNeighborModelEditor/Private/NearestNeighborEditorModel.cpp
--- a/Engine/Plugins/Animation/MLDeformer/NearestNeighborModel/Source/NearestNeighborModelEditor/Private/NearestNeighborEditorModel.cpp
+++ b/Engine/Plugins/Animation/MLDeformer/NearestNeighborModel/Source/NearestNeighborModelEditor/Private/NearestNeighborEditorModel.cpp
@@ -102,7 +102,16 @@ namespace UE::NearestNeighborModel
 				if (NearestNeighborModel->DoesUseOptimizedNetwork())
 				{
 					Model->SetNeuralNetwork(Network);
-					return NearestNeighborModel->LoadOptimizedNetwork(OnnxFile);
+					const bool bSuccess = NearestNeighborModel->LoadOptimizedNetwork(OnnxFile);
+					if (bSuccess)
+					{
+						UNearestNeighborModelInstance* ModelInstance = static_cast<UNearestNeighborModelInstance*>(GetTestMLDeformerModelInstance());
+						if (ModelInstance)
+						{
+							ModelInstance->InitOptimizedNetworkInstance();
+							return true;
+						}
+					}
 				}
 				else
 				{
