Undo //UE5/Release-Engine-Staging/Engine/Plugins/Animation/MLDeformer/NearestNeighborModel/... changelist 23801366

diff --git a/Engine/Plugins/Animation/MLDeformer/NearestNeighborModel/Source/NearestNeighborModelEditor/Private/NearestNeighborModelDetails.cpp b/Engine/Plugins/Animation/MLDeformer/NearestNeighborModel/Source/NearestNeighborModelEditor/Private/NearestNeighborModelDetails.cpp
--- a/Engine/Plugins/Animation/MLDeformer/NearestNeighborModel/Source/NearestNeighborModelEditor/Private/NearestNeighborModelDetails.cpp
+++ b/Engine/Plugins/Animation/MLDeformer/NearestNeighborModel/Source/NearestNeighborModelEditor/Private/NearestNeighborModelDetails.cpp
@@ -193,7 +193,7 @@ namespace UE::NearestNeighborModel
 
 		BuildSubMeshNames();
 		TSharedRef<IPropertyHandle> ClothPartDataPropertyHandle = DetailBuilder.GetProperty(UNearestNeighborModel::GetClothPartEditorDataPropertyName());
-		if (ClothPartDataPropertyHandle->AsArray().IsValid())
+		if (ClothPartDataPropertyHandle->AsArray().IsValid() && !SubMeshNames.IsEmpty())
 		{
 			TSharedRef<FDetailArrayBuilder> PropertyBuilder = MakeShared<FDetailArrayBuilder>(ClothPartDataPropertyHandle, true, false, true);
 			PropertyBuilder->OnGenerateArrayElementWidget(FOnGenerateArrayElementWidget::CreateSP(this, &FNearestNeighborModelDetails::GenerateClothPartElementWidget));
@@ -206,8 +206,9 @@ namespace UE::NearestNeighborModel
 		NearestNeighborCategoryBuilder->AddProperty(UNearestNeighborModel::GetUsePartOnlyMeshPropertyName());
 		NearestNeighborCategoryBuilder->AddProperty(UNearestNeighborModel::GetNearestNeighborDataPropertyName());
 
-		KMeansCategoryBuilder->AddProperty(GET_MEMBER_NAME_STRING_CHECKED(UNearestNeighborModel, SourceSkeletons));
+		KMeansCategoryBuilder->AddProperty(GET_MEMBER_NAME_STRING_CHECKED(UNearestNeighborModel, SourceAnims));
 		KMeansCategoryBuilder->AddProperty(GET_MEMBER_NAME_STRING_CHECKED(UNearestNeighborModel, NumClusters));
+		KMeansCategoryBuilder->AddProperty(GET_MEMBER_NAME_STRING_CHECKED(UNearestNeighborModel, KMeansPartId));
 		KMeansCategoryBuilder->AddCustomRow(FText::FromString(""))
 			.WholeRowContent()
 			[
@@ -217,6 +218,11 @@ namespace UE::NearestNeighborModel
 				.OnClicked_Lambda([this]
 				{
 					NearestNeighborEditorModel->KMeansClusterPoses();
+					EditorModel->GetEditor()->GetModelDetailsView()->ForceRefresh();
+					if (NearestNeighborEditorModel->GetKMeansClusterResult() == EUpdateResult::SUCCESS)
+					{
+						UE_LOG(LogNearestNeighborModel, Display, TEXT("Cluster succeeded."));
+					}
 					return FReply::Handled();
 				})
 			];
@@ -235,13 +241,18 @@ namespace UE::NearestNeighborModel
 				{
 					NearestNeighborEditorModel->OnMorphTargetUpdate();
 					EditorModel->GetEditor()->GetModelDetailsView()->ForceRefresh();
+					if (NearestNeighborEditorModel->GetMorphTargetUpdateResult() == EUpdateResult::SUCCESS)
+					{
+						UE_LOG(LogNearestNeighborModel, Display, TEXT("Update succeeded."));
+					}
 					return FReply::Handled();
 				})
 			];
-		AddUpdateResultText(MorphTargetCategoryBuilder, NearestNeighborEditorModel->GetMorphTargetUpdateResult());
+		AddActionResultText(MorphTargetCategoryBuilder, NearestNeighborEditorModel->GetMorphTargetUpdateResult(), TEXT("Update"));
+		AddActionResultText(KMeansCategoryBuilder, NearestNeighborEditorModel->GetKMeansClusterResult(), TEXT("KMeans"));
 	}
 
-	void FNearestNeighborModelDetails::AddUpdateResultText(IDetailCategoryBuilder* CategoryBuilder, uint8 Result)
+	void FNearestNeighborModelDetails::AddActionResultText(IDetailCategoryBuilder* CategoryBuilder, uint8 Result, const FString& ActionName)
 	{
 		CategoryBuilder->AddCustomRow(FText::FromString("UpdateResultError"))
 			.Visibility((Result & EUpdateResult::ERROR) != 0 ? EVisibility::Visible : EVisibility::Collapsed)
@@ -252,7 +263,7 @@ namespace UE::NearestNeighborModel
 				[
 					SNew(SWarningOrErrorBox)
 					.MessageStyle(EMessageStyle::Error)
-					.Message(FText::FromString(TEXT("Update failed with errors. Please check Output Log (LogNearestNeighborModel, LogPython) for details.")))
+					.Message(FText::FromString(ActionName + TEXT(" failed with errors. Please check Output Log (LogNearestNeighborModel, LogPython) for details.")))
 				]
 			];
 		CategoryBuilder->AddCustomRow(FText::FromString("UpdateResultWarning"))
@@ -264,13 +275,9 @@ namespace UE::NearestNeighborModel
 				[
 					SNew(SWarningOrErrorBox)
 					.MessageStyle(EMessageStyle::Warning)
-					.Message(FText::FromString(TEXT("Update finished with warnings. Please check Output Log (LogNearestNeighborModel, LogPython) for details.")))
+					.Message(FText::FromString(ActionName + TEXT(" finished with warnings. Please check Output Log (LogNearestNeighborModel, LogPython) for details.")))
 				]
 			];
-		if (Result == EUpdateResult::SUCCESS)
-		{
-			UE_LOG(LogNearestNeighborModel, Display, TEXT("Update succeeded."));
-		}
 	}
 
 	void FNearestNeighborModelDetails::BuildSubMeshNames()
diff --git a/Engine/Plugins/Animation/MLDeformer/NearestNeighborModel/Source/NearestNeighborModelEditor/Private/NearestNeighborTrainingModel.cpp b/Engine/Plugins/Animation/MLDeformer/NearestNeighborModel/Source/NearestNeighborModelEditor/Private/NearestNeighborTrainingModel.cpp
--- a/Engine/Plugins/Animation/MLDeformer/NearestNeighborModel/Source/NearestNeighborModelEditor/Private/NearestNeighborTrainingModel.cpp
+++ b/Engine/Plugins/Animation/MLDeformer/NearestNeighborModel/Source/NearestNeighborModelEditor/Private/NearestNeighborTrainingModel.cpp
@@ -52,30 +52,35 @@ int32 UNearestNeighborTrainingModel::GetPartNumNeighbors(const int32 PartId) con
 	return FMath::Min(NearestNeighborModel->GetNumNeighborsFromAnimSequence(PartId), NearestNeighborModel->GetNumNeighborsFromGeometryCache(PartId));
 }
 
-void UNearestNeighborTrainingModel::SampleKmeansAnim(const int32 SkeletonId)
+bool UNearestNeighborTrainingModel::SampleKmeansAnim(const int32 SkeletonId)
 {
 	FNearestNeighborGeomCacheSampler* Sampler = static_cast<FNearestNeighborGeomCacheSampler*>(EditorModel->GetSampler());
-	Sampler->SampleKMeansAnim(SkeletonId);
+	return Sampler->SampleKMeansAnim(SkeletonId);
 }
 
-void UNearestNeighborTrainingModel::SampleKmeansFrame(const int32 Frame)
+bool UNearestNeighborTrainingModel::SampleKmeansFrame(const int32 Frame)
 {
 	FNearestNeighborGeomCacheSampler* Sampler = static_cast<FNearestNeighborGeomCacheSampler*>(EditorModel->GetSampler());
-	Sampler->SampleKMeansFrame(Frame);
-	SampleBoneRotations = Sampler->GetBoneRotations();
+	const bool bSampleExist = Sampler->SampleKMeansFrame(Frame);
+	if (bSampleExist)
+	{
+		SampleBoneRotations = Sampler->GetBoneRotations();
+		return true;	
+	}
+	return false;
 }
 
 
 int32 UNearestNeighborTrainingModel::GetKmeansNumAnims() const
 {
-	return NearestNeighborModel->SourceSkeletons.Num();
+	return NearestNeighborModel->SourceAnims.Num();
 }
 
 int32 UNearestNeighborTrainingModel::GetKmeansAnimNumFrames(const int32 SkeletonId) const
 {
-	if (SkeletonId < NearestNeighborModel->SourceSkeletons.Num())
+	if (SkeletonId < NearestNeighborModel->SourceAnims.Num())
 	{
-		return NearestNeighborModel->SourceSkeletons[SkeletonId]->GetDataModel()->GetNumberOfFrames();
+		return NearestNeighborModel->SourceAnims[SkeletonId]->GetDataModel()->GetNumberOfFrames();
 	}
 	return 0;
 }
diff --git a/Engine/Plugins/Animation/MLDeformer/NearestNeighborModel/Source/NearestNeighborModelEditor/Public/NearestNeighborEditorModel.h b/Engine/Plugins/Animation/MLDeformer/NearestNeighborModel/Source/NearestNeighborModelEditor/Public/NearestNeighborEditorModel.h
--- a/Engine/Plugins/Animation/MLDeformer/NearestNeighborModel/Source/NearestNeighborModelEditor/Public/NearestNeighborEditorModel.h
+++ b/Engine/Plugins/Animation/MLDeformer/NearestNeighborModel/Source/NearestNeighborModelEditor/Public/NearestNeighborEditorModel.h
@@ -65,6 +65,8 @@ namespace UE::NearestNeighborModel
 
 		void UpdateNearestNeighborActors();
 		void KMeansClusterPoses();
+		TPair<UAnimSequence*, uint8> CreateAnimOfClusterCenters(const FString& PackageName, const TArray<int32>& KmeansResults);
+		uint8 GetKMeansClusterResult() const { return KMeansClusterResult; }
 
 		uint8 InitMorphTargets();
 		void RefreshMorphTargets();
@@ -103,9 +105,8 @@ namespace UE::NearestNeighborModel
 
 	private:
 		UWorld* EditorWorld = nullptr;
-		uint8 ClothPartUpdateResult = 0;
-		uint8 NearestNeighborUpdateResult = 0;
 		uint8 MorphTargetUpdateResult = 0;
+		uint8 KMeansClusterResult = 0;
 		int32 NumTrainingFramesOverride = -1;
 	};
 }	// namespace UE::NearestNeighborModel
diff --git a/Engine/Plugins/Animation/MLDeformer/NearestNeighborModel/Source/NearestNeighborModelEditor/Public/NearestNeighborGeomCacheSampler.h b/Engine/Plugins/Animation/MLDeformer/NearestNeighborModel/Source/NearestNeighborModelEditor/Public/NearestNeighborGeomCacheSampler.h
--- a/Engine/Plugins/Animation/MLDeformer/NearestNeighborModel/Source/NearestNeighborModelEditor/Public/NearestNeighborGeomCacheSampler.h
+++ b/Engine/Plugins/Animation/MLDeformer/NearestNeighborModel/Source/NearestNeighborModelEditor/Public/NearestNeighborGeomCacheSampler.h
@@ -11,8 +11,8 @@ namespace UE::NearestNeighborModel
 	{
 	public:
 		virtual uint8 SamplePart(int32 InAnimFrameIndex, int32 PartId);
-		virtual void SampleKMeansAnim(const int32 SkeletonId);
-		virtual void SampleKMeansFrame(const int32 Frame);
+		virtual bool SampleKMeansAnim(const int32 SkeletonId);
+		virtual bool SampleKMeansFrame(const int32 Frame);
 		const TArray<float>& GetPartVertexDeltas() const { return PartVertexDeltas; }
 		uint8 GeneratePartMeshMappings(const TArray<uint32>& VertexMap, bool bUsePartOnlyMesh);
 		uint8 GenerateMeshMappingIndices();
@@ -22,5 +22,6 @@ namespace UE::NearestNeighborModel
 	protected:
 		TArray<float> PartVertexDeltas;
 		TArray<int32> MeshMappingIndices;
+		int32 KMeansAnimId;
 	};
 }
diff --git a/Engine/Plugins/Animation/MLDeformer/NearestNeighborModel/Source/NearestNeighborModelEditor/Public/NearestNeighborModelDetails.h b/Engine/Plugins/Animation/MLDeformer/NearestNeighborModel/Source/NearestNeighborModelEditor/Public/NearestNeighborModelDetails.h
--- a/Engine/Plugins/Animation/MLDeformer/NearestNeighborModel/Source/NearestNeighborModelEditor/Public/NearestNeighborModelDetails.h
+++ b/Engine/Plugins/Animation/MLDeformer/NearestNeighborModel/Source/NearestNeighborModelEditor/Public/NearestNeighborModelDetails.h
@@ -8,6 +8,7 @@
 class UNearestNeighborModel;
 class IPropertyHandle;
 class IDetailChildrenBuilder;
+class FDetailWidgetRow;
 
 namespace UE::NearestNeighborModel
 {
@@ -27,7 +28,7 @@ namespace UE::NearestNeighborModel
 		virtual void CreateCategories() override;
 
 		virtual bool UpdateMemberPointers(const TArray<TWeakObjectPtr<UObject>>& Objects) override;
-		virtual void AddUpdateResultText(IDetailCategoryBuilder* CategoryBuilder, uint8 Result);
+		virtual void AddActionResultText(IDetailCategoryBuilder* CategoryBuilder, uint8 Result, const FString& ActionName);
 
 		void GenerateClothPartElementWidget(TSharedRef<IPropertyHandle> PropertyHandle, int32 ArrayIndex, IDetailChildrenBuilder& ChildrenBuilder);
 
diff --git a/Engine/Plugins/Animation/MLDeformer/NearestNeighborModel/Source/NearestNeighborModelEditor/Public/NearestNeighborTrainingModel.h b/Engine/Plugins/Animation/MLDeformer/NearestNeighborModel/Source/NearestNeighborModelEditor/Public/NearestNeighborTrainingModel.h
--- a/Engine/Plugins/Animation/MLDeformer/NearestNeighborModel/Source/NearestNeighborModelEditor/Public/NearestNeighborTrainingModel.h
+++ b/Engine/Plugins/Animation/MLDeformer/NearestNeighborModel/Source/NearestNeighborModelEditor/Public/NearestNeighborTrainingModel.h
@@ -47,13 +47,13 @@ public:
 	int32 GetPartNumNeighbors(const int32 PartId) const;
 
 	UFUNCTION(BlueprintImplementableEvent, Category = "Python")
-	void KmeansClusterPoses() const;
+	int32 KmeansClusterPoses(const int32 PartId) const;
 
 	UFUNCTION(BlueprintCallable, Category = "Python")
-	void SampleKmeansAnim(const int32 SkeletonId);
+	bool SampleKmeansAnim(const int32 SkeletonId);
 
 	UFUNCTION(BlueprintCallable, Category = "Python")
-	void SampleKmeansFrame(const int32 Frame);
+	bool SampleKmeansFrame(const int32 Frame);
 
 	UFUNCTION(BlueprintPure, Category = "Python")
 	int32 GetKmeansNumAnims() const;
@@ -69,5 +69,8 @@ public:
 	UPROPERTY(EditAnywhere, BlueprintReadWrite, Category = "Training Data")
 	TArray<float> PartSampleDeltas;
 
+	UPROPERTY(EditAnywhere, BlueprintReadWrite, Category = "Training Data")
+	TArray<int32> KmeansResults;
+
 	UNearestNeighborModel* NearestNeighborModel = nullptr;
 };
