NearestNeighborModel: Fix function comment

diff --git a/Engine/Plugins/Animation/MLDeformer/NearestNeighborModel/Source/NearestNeighborModelEditor/Private/NearestNeighborEditorModel.cpp b/Engine/Plugins/Animation/MLDeformer/NearestNeighborModel/Source/NearestNeighborModelEditor/Private/NearestNeighborEditorModel.cpp
--- a/Engine/Plugins/Animation/MLDeformer/NearestNeighborModel/Source/NearestNeighborModelEditor/Private/NearestNeighborEditorModel.cpp
+++ b/Engine/Plugins/Animation/MLDeformer/NearestNeighborModel/Source/NearestNeighborModelEditor/Private/NearestNeighborEditorModel.cpp
@@ -450,8 +450,6 @@ namespace UE::NearestNeighborModel
 		return EUpdateResult::ERROR;
 	}
 
-	// write a function to return an animation sequence and a boolean
-
 	template<typename T>
 	T* CreateObjectInstance(const FString& PackageName)
 	{
