NearestNeighborModel: Fix access order of the nearest neighbor coefficients

diff --git a/Engine/Plugins/Animation/MLDeformer/NearestNeighborModel/Source/NearestNeighborModel/Private/NearestNeighborModelInstance.cpp b/Engine/Plugins/Animation/MLDeformer/NearestNeighborModel/Source/NearestNeighborModel/Private/NearestNeighborModelInstance.cpp
--- a/Engine/Plugins/Animation/MLDeformer/NearestNeighborModel/Source/NearestNeighborModel/Private/NearestNeighborModelInstance.cpp
+++ b/Engine/Plugins/Animation/MLDeformer/NearestNeighborModel/Source/NearestNeighborModel/Private/NearestNeighborModelInstance.cpp
@@ -237,9 +237,9 @@ int32 UNearestNeighborModelInstance::FindNearestNeighbor(const float* PCAData, i
 		const int32 NumNeighbors = NearestNeighborModel->GetNumNeighbors(PartId);
 		for (int32 CoeffId = 0; CoeffId < NumCoeffs; CoeffId++)
 		{
-			float Coeff = PCAData[CoeffStart + CoeffId];
-			float NeighborCoeff = NeighborCoeffs[CoeffId * NumNeighbors + NeighborId];
-			float D = Coeff - NeighborCoeff;
+			const float Coeff = PCAData[CoeffStart + CoeffId];
+			const float NeighborCoeff = NeighborCoeffs[NeighborId * NumCoeffs + CoeffId];
+			const float D = Coeff - NeighborCoeff;
 			D2 += D * D;
 		}
 		if (D2 < MinD2)
