NearestNeighborModel: Fix python package dependencies for optimized network

diff --git a/Engine/Plugins/Animation/MLDeformer/NearestNeighborModel/Source/NearestNeighborModel/Private/NearestNeighborModel.cpp b/Engine/Plugins/Animation/MLDeformer/NearestNeighborModel/Source/NearestNeighborModel/Private/NearestNeighborModel.cpp
--- a/Engine/Plugins/Animation/MLDeformer/NearestNeighborModel/Source/NearestNeighborModel/Private/NearestNeighborModel.cpp
+++ b/Engine/Plugins/Animation/MLDeformer/NearestNeighborModel/Source/NearestNeighborModel/Private/NearestNeighborModel.cpp
@@ -47,6 +47,7 @@ UNearestNeighborModel::UNearestNeighborModel(const FObjectInitializer& ObjectIni
 {
 #if WITH_EDITORONLY_DATA
 	SetVizSettings(ObjectInitializer.CreateEditorOnlyDefaultSubobject<UNearestNeighborModelVizSettings>(this, TEXT("VizSettings")));
+	bDoesMeetOptimizedNetworkPrerequisites = DoesMeetOptimizedNetworkPrerequisites();
 #endif
 }
 
@@ -447,6 +448,20 @@ FString UNearestNeighborModel::GetModelDir() const
 		return FPaths::ProjectIntermediateDir() + TEXT("NearestNeighborModel/");
 	}
 }
+
+bool UNearestNeighborModel::DoesEditorSupportOptimizedNetwork() const
+{
+#if NEARESTNEIGHBORMODEL_USE_ISPC
+	return bDoesMeetOptimizedNetworkPrerequisites && UE::NearestNeighborModel::bNearestNeighborModelUseOptimizedNetwork;
+#else
+	return false;
+#endif
+}
+
+void UNearestNeighborModel::SetUseOptimizedNetwork(bool bInUseOptimizedNetwork)
+{
+	bUseOptimizedNetwork = bInUseOptimizedNetwork;
+}
 #endif
 
 void UNearestNeighborModel::InitInputInfo()
@@ -470,13 +485,10 @@ void UNearestNeighborModel::SetOptimizedNetwork(UNearestNeighborOptimizedNetwork
 	OptimizedNetwork = InOptimizedNetwork;
 }
 
+
 bool UNearestNeighborModel::DoesUseOptimizedNetwork() const
 {
-#if NEARESTNEIGHBORMODEL_USE_ISPC
-	return UE::NearestNeighborModel::bNearestNeighborModelUseOptimizedNetwork;
-#else
-	return false;
-#endif
+	return bUseOptimizedNetwork;
 }
 
 /**
@@ -531,4 +543,14 @@ int32 UNearestNeighborModel::GetOptimizedNetworkNumOutputs() const
 	return OptimizedNetwork ? OptimizedNetwork->GetNumOutputs() : 0;
 }
 
+bool UNearestNeighborModel::DoesMeetOptimizedNetworkPrerequisites() const
+{
+	UNearestNeighborOptimizedNetworkLoader* Loader = GetDerivedCDO<UNearestNeighborOptimizedNetworkLoader>();
+	if (Loader != nullptr)
+	{
+		return Loader->DoesMeetPrerequisites();
+	}
+	return false;
+}
+
 #undef LOCTEXT_NAMESPACE
diff --git a/Engine/Plugins/Animation/MLDeformer/NearestNeighborModel/Source/NearestNeighborModel/Public/NearestNeighborModel.h b/Engine/Plugins/Animation/MLDeformer/NearestNeighborModel/Source/NearestNeighborModel/Public/NearestNeighborModel.h
--- a/Engine/Plugins/Animation/MLDeformer/NearestNeighborModel/Source/NearestNeighborModel/Public/NearestNeighborModel.h
+++ b/Engine/Plugins/Animation/MLDeformer/NearestNeighborModel/Source/NearestNeighborModel/Public/NearestNeighborModel.h
@@ -239,6 +239,7 @@ public:
 	virtual int32 GetNumFloatsPerCurve() const override { return NearestNeighborNumFloatsPerCurve; }
 
 	bool DoesUseOptimizedNetwork() const;
+	bool DoesMeetOptimizedNetworkPrerequisites() const;
 	UNearestNeighborOptimizedNetwork* GetOptimizedNetwork() const { return OptimizedNetwork.Get(); }
 	int32 GetOptimizedNetworkNumOutputs() const;
 	void SetOptimizedNetwork(UNearestNeighborOptimizedNetwork* InOptimizedNetwork);
@@ -273,6 +274,8 @@ public:
 	bool IsClothPartDataValid() const { return bClothPartDataValid; }
 	bool IsNearestNeighborDataValid() const { return bNearestNeighborDataValid; }
 	bool IsMorphTargetDataValid() const { return bMorphTargetDataValid; }
+	bool DoesEditorSupportOptimizedNetwork() const;
+	void SetUseOptimizedNetwork(bool bInUseOptimizedNetwork);
 #endif
 
 #if WITH_EDITOR
@@ -430,4 +433,10 @@ public:
 	/** The NNI neural network. */
 	UPROPERTY()
 	TObjectPtr<UNeuralNetwork> NNINetwork;
+
+	UPROPERTY()
+	bool bDoesMeetOptimizedNetworkPrerequisites = false;
+
+	UPROPERTY()
+	bool bUseOptimizedNetwork = true;
 };
diff --git a/Engine/Plugins/Animation/MLDeformer/NearestNeighborModel/Source/NearestNeighborModel/Public/NearestNeighborOptimizedNetworkLoader.h b/Engine/Plugins/Animation/MLDeformer/NearestNeighborModel/Source/NearestNeighborModel/Public/NearestNeighborOptimizedNetworkLoader.h
--- a/Engine/Plugins/Animation/MLDeformer/NearestNeighborModel/Source/NearestNeighborModel/Public/NearestNeighborOptimizedNetworkLoader.h
+++ b/Engine/Plugins/Animation/MLDeformer/NearestNeighborModel/Source/NearestNeighborModel/Public/NearestNeighborOptimizedNetworkLoader.h
@@ -20,6 +20,9 @@ public:
 	UFUNCTION(BlueprintImplementableEvent, Category = "Nearest Neighbor Model")
 	bool LoadOptimizedNetwork(const FString& OnnxPath);
 
+	UFUNCTION(BlueprintImplementableEvent, Category = "Nearest Neighbor Model")
+	bool DoesMeetPrerequisites() const;
+
 	void SetOptimizedNetwork(UNearestNeighborOptimizedNetwork* InNetwork);
 
 	UFUNCTION(BlueprintPure, Category = "Nearest Neighbor Model")
diff --git a/Engine/Plugins/Animation/MLDeformer/NearestNeighborModel/Source/NearestNeighborModelEditor/Private/NearestNeighborEditorModel.cpp b/Engine/Plugins/Animation/MLDeformer/NearestNeighborModel/Source/NearestNeighborModelEditor/Private/NearestNeighborEditorModel.cpp
--- a/Engine/Plugins/Animation/MLDeformer/NearestNeighborModel/Source/NearestNeighborModelEditor/Private/NearestNeighborEditorModel.cpp
+++ b/Engine/Plugins/Animation/MLDeformer/NearestNeighborModel/Source/NearestNeighborModelEditor/Private/NearestNeighborEditorModel.cpp
@@ -130,9 +130,11 @@ namespace UE::NearestNeighborModel
 			UNeuralNetwork* Network = LoadNeuralNetworkFromOnnx(OnnxFile);
 			if (Network)
 			{
-				if (NearestNeighborModel->DoesUseOptimizedNetwork())
+				if (NearestNeighborModel->DoesEditorSupportOptimizedNetwork())
 				{
+					// We still need NNINetwork for serialization in nearest neighbor update and kmeans update.
 					NearestNeighborModel->SetNNINetwork(Network);
+
 					const bool bSuccess = NearestNeighborModel->LoadOptimizedNetwork(OnnxFile);
 					if (bSuccess)
 					{
@@ -143,10 +145,12 @@ namespace UE::NearestNeighborModel
 							return true;
 						}
 					}
+					NearestNeighborModel->SetUseOptimizedNetwork(true);
 				}
 				else
 				{
 					NearestNeighborModel->SetNNINetwork(Network);
+					NearestNeighborModel->SetUseOptimizedNetwork(false);
 					return true;
 				}
 			}
@@ -274,8 +278,8 @@ namespace UE::NearestNeighborModel
 	void FNearestNeighborEditorModel::InitInputInfo(UMLDeformerInputInfo* InputInfo)
 	{
 		FMLDeformerEditorModel::InitInputInfo(InputInfo);
-		UNearestNeighborModel *NearestNeighborModel = static_cast<UNearestNeighborModel*>(Model);
-		NearestNeighborModel->InitInputInfo();
+		UNearestNeighborModelInputInfo* NearestNeighborInputInfo = static_cast<UNearestNeighborModelInputInfo*>(InputInfo);
+		NearestNeighborInputInfo->InitRefBoneRotations(Model->GetSkeletalMesh());
 	}
 
 	ETrainingResult FNearestNeighborEditorModel::Train()
@@ -718,9 +722,9 @@ namespace UE::NearestNeighborModel
 		MorphTargetUpdateResult = EUpdateResult::SUCCESS;
 
 		UNearestNeighborModel* NearestNeighborModel = GetNearestNeighborModel();
-		// Convert NNI network to  optimized network
+		// Automatic converting NNI network to optimized network
 		const UNeuralNetwork* NeuralNetwork = NearestNeighborModel->GetNNINetwork();
-		if (NearestNeighborModel->DoesUseOptimizedNetwork() && NeuralNetwork != nullptr && NearestNeighborModel->GetOptimizedNetwork() == nullptr)
+		if (NearestNeighborModel->DoesEditorSupportOptimizedNetwork() && NeuralNetwork != nullptr && NearestNeighborModel->GetOptimizedNetwork() == nullptr)
 		{
 			const FString SavePath = GetTrainedNetworkOnnxFile();
 			UE_LOG(LogNearestNeighborModel, Display, TEXT("Saving to %s"), *SavePath);
@@ -903,6 +907,12 @@ namespace UE::NearestNeighborModel
 
 		if (NearestNeighborModel->DoesUseOptimizedNetwork())
 		{
+			if (NearestNeighborModel->GetOptimizedNetwork() == nullptr)
+			{
+				UE_LOG(LogNearestNeighborModel, Warning, TEXT("Optimized network is not set. Model needs to be re-trained."));
+				return EUpdateResult::WARNING;
+			}
+
 			const int32 NumNetworkWeights = NearestNeighborModel->GetOptimizedNetworkNumOutputs();
 			const int32 NumPCACoeffs = NearestNeighborModel->GetTotalNumPCACoeffs();
 			if (NumNetworkWeights != NumPCACoeffs)
