NearestNeighorModel: Fix crash if the skeletal mesh is incompatible with the existing model

diff --git a/Engine/Plugins/Animation/MLDeformer/NearestNeighborModel/Source/NearestNeighborModel/Private/NearestNeighborModel.cpp b/Engine/Plugins/Animation/MLDeformer/NearestNeighborModel/Source/NearestNeighborModel/Private/NearestNeighborModel.cpp
--- a/Engine/Plugins/Animation/MLDeformer/NearestNeighborModel/Source/NearestNeighborModel/Private/NearestNeighborModel.cpp
+++ b/Engine/Plugins/Animation/MLDeformer/NearestNeighborModel/Source/NearestNeighborModel/Private/NearestNeighborModel.cpp
@@ -322,7 +322,10 @@ UE::NearestNeighborModel::EUpdateResult UNearestNeighborModel::UpdateClothPartDa
 	for(int32 PartId = 0; PartId < GetNumParts(); PartId++)
 	{
 		ClothPartData[PartId].PCACoeffNum = ClothPartEditorData[PartId].PCACoeffNum;
-		ReturnCode |= UpdateVertexMap(PartId, ClothPartEditorData[PartId].VertexMapPath, SkelMeshInfos[ClothPartEditorData[PartId].MeshIndex]);
+		int32& MeshIndex = ClothPartEditorData[PartId].MeshIndex;
+		// Set to 0 if MeshIndex is invalid. This could happen if the skeletal mesh is changed.
+		MeshIndex = MeshIndex < SkelMeshInfos.Num() ? MeshIndex : 0;
+		ReturnCode |= UpdateVertexMap(PartId, ClothPartEditorData[PartId].VertexMapPath, SkelMeshInfos[MeshIndex]);
 		ClothPartData[PartId].NumVertices = ClothPartData[PartId].VertexMap.Num();
 
 		if (!CheckPCAData(PartId))
@@ -462,6 +465,16 @@ void UNearestNeighborModel::SetUseOptimizedNetwork(bool bInUseOptimizedNetwork)
 {
 	bUseOptimizedNetwork = bInUseOptimizedNetwork;
 }
+
+int32 UNearestNeighborModel::GetMaxPartMeshIndex() const
+{
+	int32 MaxIndex = -1;
+	for (const FClothPartEditorData& PartData : ClothPartEditorData)
+	{
+		MaxIndex = FMath::Max(MaxIndex, PartData.MeshIndex);
+	}
+	return MaxIndex;
+}
 #endif
 
 void UNearestNeighborModel::InitInputInfo()
diff --git a/Engine/Plugins/Animation/MLDeformer/NearestNeighborModel/Source/NearestNeighborModel/Public/NearestNeighborModel.h b/Engine/Plugins/Animation/MLDeformer/NearestNeighborModel/Source/NearestNeighborModel/Public/NearestNeighborModel.h
--- a/Engine/Plugins/Animation/MLDeformer/NearestNeighborModel/Source/NearestNeighborModel/Public/NearestNeighborModel.h
+++ b/Engine/Plugins/Animation/MLDeformer/NearestNeighborModel/Source/NearestNeighborModel/Public/NearestNeighborModel.h
@@ -331,6 +331,7 @@ public:
 
 	int32 GetPartMeshIndex(int32 PartId) const { return ClothPartEditorData[PartId].MeshIndex; }
 	void SetPartMeshIndex(int32 PartId, int32 MeshIndex) { ClothPartEditorData[PartId].MeshIndex = MeshIndex; }
+	int32 GetMaxPartMeshIndex() const;
 #endif
 
 
diff --git a/Engine/Plugins/Animation/MLDeformer/NearestNeighborModel/Source/NearestNeighborModelEditor/Private/NearestNeighborModelDetails.cpp b/Engine/Plugins/Animation/MLDeformer/NearestNeighborModel/Source/NearestNeighborModelEditor/Private/NearestNeighborModelDetails.cpp
--- a/Engine/Plugins/Animation/MLDeformer/NearestNeighborModel/Source/NearestNeighborModelEditor/Private/NearestNeighborModelDetails.cpp
+++ b/Engine/Plugins/Animation/MLDeformer/NearestNeighborModel/Source/NearestNeighborModelEditor/Private/NearestNeighborModelDetails.cpp
@@ -78,7 +78,8 @@ namespace UE::NearestNeighborModel
 
 		typename STextComboBox::FOnTextSelectionChanged SubMeshComboOnSelectionChangedDelegate;
 		SubMeshComboOnSelectionChangedDelegate.BindRaw(this, &FNearestNeighborModelDetails::SubMeshComboSelectionChanged, ArrayIndex);
-		const int32 InitMeshIndex = NearestNeighborModel ? NearestNeighborModel->GetPartMeshIndex(ArrayIndex) : 0;
+		int32 InitMeshIndex = NearestNeighborModel ? NearestNeighborModel->GetPartMeshIndex(ArrayIndex) : 0;
+		InitMeshIndex = InitMeshIndex >= SubMeshNames.Num() ? 0 : InitMeshIndex;
 
 		ChildrenBuilder.AddCustomRow(FText::GetEmpty())
 		[
@@ -191,7 +192,20 @@ namespace UE::NearestNeighborModel
 		Group->AddPropertyRow(DetailBuilder.GetProperty(UNearestNeighborModel::GetRecomputeDeltasPropertyName()));
 		Group->AddPropertyRow(DetailBuilder.GetProperty(UNearestNeighborModel::GetRecomputePCAPropertyName()));
 
+		if (NearestNeighborModel == nullptr)
+		{
+			return;
+		}
+
 		BuildSubMeshNames();
+		const int32 MaxPartMeshIndex = NearestNeighborModel->GetMaxPartMeshIndex();
+		if (MaxPartMeshIndex >= SubMeshNames.Num())
+		{
+			UE_LOG(LogNearestNeighborModel, Error, TEXT("Nearest neighbor model was previously created with %d submeshes, but the current skeletal mesh has %d submeshes. Please use the original skeletal mesh or click update to overwrite existing data."), MaxPartMeshIndex, SubMeshNames.Num());
+
+			NearestNeighborModel->InvalidateClothPartData();
+			AddActionResultText(ClothPartCategoryBuilder, EUpdateResult::ERROR, TEXT("Loading"));
+		}
 		TSharedRef<IPropertyHandle> ClothPartDataPropertyHandle = DetailBuilder.GetProperty(UNearestNeighborModel::GetClothPartEditorDataPropertyName());
 		if (ClothPartDataPropertyHandle->AsArray().IsValid() && !SubMeshNames.IsEmpty())
 		{
