Added CSV Stats to MLDeformer

diff --git a/Engine/Plugins/Animation/MLDeformer/MLDeformerFramework/Source/MLDeformerFramework/Private/MLDeformerModelInstance.cpp b/Engine/Plugins/Animation/MLDeformer/MLDeformerFramework/Source/MLDeformerFramework/Private/MLDeformerModelInstance.cpp
--- a/Engine/Plugins/Animation/MLDeformer/MLDeformerFramework/Source/MLDeformerFramework/Private/MLDeformerModelInstance.cpp
+++ b/Engine/Plugins/Animation/MLDeformer/MLDeformerFramework/Source/MLDeformerFramework/Private/MLDeformerModelInstance.cpp
@@ -14,6 +14,8 @@
 
 #include UE_INLINE_GENERATED_CPP_BY_NAME(MLDeformerModelInstance)
 
+CSV_DEFINE_CATEGORY_MODULE(MLDEFORMERFRAMEWORK_API, MLDeformer, false);
+
 void UMLDeformerModelInstance::BeginDestroy()
 {
 	Release();
diff --git a/Engine/Plugins/Animation/MLDeformer/MLDeformerFramework/Source/MLDeformerFramework/Public/MLDeformerModelInstance.h b/Engine/Plugins/Animation/MLDeformer/MLDeformerFramework/Source/MLDeformerFramework/Public/MLDeformerModelInstance.h
--- a/Engine/Plugins/Animation/MLDeformer/MLDeformerFramework/Source/MLDeformerFramework/Public/MLDeformerModelInstance.h
+++ b/Engine/Plugins/Animation/MLDeformer/MLDeformerFramework/Source/MLDeformerFramework/Public/MLDeformerModelInstance.h
@@ -6,6 +6,7 @@
 #include "MLDeformerModule.h"
 #include "UObject/ObjectPtr.h"
 #include "RenderCommandFence.h"
+#include "ProfilingDebugging/CsvProfiler.h"
 #include "MLDeformerModelInstance.generated.h"
 
 class UMLDeformerModel;
diff --git a/Engine/Plugins/Animation/MLDeformer/NearestNeighborModel/Source/NearestNeighborModel/Private/NearestNeighborModelInstance.cpp b/Engine/Plugins/Animation/MLDeformer/NearestNeighborModel/Source/NearestNeighborModel/Private/NearestNeighborModelInstance.cpp
--- a/Engine/Plugins/Animation/MLDeformer/NearestNeighborModel/Source/NearestNeighborModel/Private/NearestNeighborModelInstance.cpp
+++ b/Engine/Plugins/Animation/MLDeformer/NearestNeighborModel/Source/NearestNeighborModel/Private/NearestNeighborModelInstance.cpp
@@ -11,6 +11,7 @@
 #include "Components/SkeletalMeshComponent.h"
 
 #include UE_INLINE_GENERATED_CPP_BY_NAME(NearestNeighborModelInstance)
+CSV_DECLARE_CATEGORY_MODULE_EXTERN(MLDEFORMERFRAMEWORK_API, MLDeformer);
 
 using namespace UE::MLDeformer;
 using namespace UE::NearestNeighborModel;
@@ -119,6 +120,7 @@ void UNearestNeighborModelInstance::InitOptimizedNetworkInstance()
 
 void UNearestNeighborModelInstance::Execute(float ModelWeight)
 {
+	CSV_SCOPED_TIMING_STAT(MLDeformer, NearestNeighborExecute);
 	const UNearestNeighborModel* NearestNeighborModel = Cast<UNearestNeighborModel>(Model);
 	check(NearestNeighborModel);
 	if (NearestNeighborModel->DoesUseOptimizedNetwork())
diff --git a/Engine/Plugins/Animation/MLDeformer/NeuralMorphModel/Source/NeuralMorphModel/Private/NeuralMorphModelInstance.cpp b/Engine/Plugins/Animation/MLDeformer/NeuralMorphModel/Source/NeuralMorphModel/Private/NeuralMorphModelInstance.cpp
--- a/Engine/Plugins/Animation/MLDeformer/NeuralMorphModel/Source/NeuralMorphModel/Private/NeuralMorphModelInstance.cpp
+++ b/Engine/Plugins/Animation/MLDeformer/NeuralMorphModel/Source/NeuralMorphModel/Private/NeuralMorphModelInstance.cpp
@@ -10,6 +10,8 @@
 #include "Components/ExternalMorphSet.h"
 #include "Components/SkeletalMeshComponent.h"
 
+CSV_DECLARE_CATEGORY_MODULE_EXTERN(MLDEFORMERFRAMEWORK_API, MLDeformer);
+
 void UNeuralMorphModelInstance::Init(USkeletalMeshComponent* SkelMeshComponent)
 {
 	Super::Init(SkelMeshComponent);
@@ -218,6 +220,7 @@ bool UNeuralMorphModelInstance::SetupInputs()
 void UNeuralMorphModelInstance::Execute(float ModelWeight)
 {
 	TRACE_CPUPROFILER_EVENT_SCOPE(UNeuralMorphModelInstance::Execute)
+	CSV_SCOPED_TIMING_STAT(MLDeformer, NeuralMorphExecute);
 
 	UNeuralMorphModel* MorphModel = Cast<UNeuralMorphModel>(Model);
 	UNeuralMorphNetwork* MorphNetwork = MorphModel->GetNeuralMorphNetwork();
