NearestNeighborModel: remove vertex map txt checks to allow indices spanning more than one submesh

diff --git a/Engine/Plugins/Animation/MLDeformer/NearestNeighborModel/Source/NearestNeighborModel/Private/NearestNeighborModel.cpp b/Engine/Plugins/Animation/MLDeformer/NearestNeighborModel/Source/NearestNeighborModel/Private/NearestNeighborModel.cpp
--- a/Engine/Plugins/Animation/MLDeformer/NearestNeighborModel/Source/NearestNeighborModel/Private/NearestNeighborModel.cpp
+++ b/Engine/Plugins/Animation/MLDeformer/NearestNeighborModel/Source/NearestNeighborModel/Private/NearestNeighborModel.cpp
@@ -272,30 +272,25 @@ UE::NearestNeighborModel::EUpdateResult UNearestNeighborModel::UpdateVertexMap(i
 		ReturnCode |= EUpdateResult::ERROR;
 		UE_LOG(LogNearestNeighborModel, Error, TEXT("Part %d txt path %s does not exist"), PartId, *VertexMapPath, NumVertices);
 	}
+	else if (GetSkeletalMesh() == nullptr)
+	{
+		ReturnCode |= EUpdateResult::ERROR;
+		UE_LOG(LogNearestNeighborModel, Error, TEXT("SkeletalMesh is None"));
+	}
 	else
 	{
 		const TArray<uint32> PartVertexMap = ReadTxt(ClothPartEditorData[PartId].VertexMapPath);
-		if (PartVertexMap.Num() > NumVertices)
+		const uint32 MaxVertexIndex = FMath::Max(PartVertexMap);
+		if ((int32)MaxVertexIndex >= GetSkeletalMesh()->GetNumImportedVertices())
 		{
-			bIsVertexMapValid = false;
-			UE_LOG(LogNearestNeighborModel, Error, TEXT("Part %d vertex map has %d vertices, larger than %d vertices in skeletal mesh, using %d vertices instead"), PartId, PartVertexMap.Num(), NumVertices, NumVertices);
 			ReturnCode |= EUpdateResult::ERROR;
+			UE_LOG(LogNearestNeighborModel, Error, TEXT("Part %d vertex map contains invalid vertex index %d, max vertex index is %d"), PartId, MaxVertexIndex, GetSkeletalMesh()->GetNumImportedVertices() - 1);
 		}
 		else
 		{
-			const int32 MaxIndex = FMath::Max(PartVertexMap);
-			if (MaxIndex >= NumVertices)
-			{
-				bIsVertexMapValid = false;
-				UE_LOG(LogNearestNeighborModel, Error, TEXT("Part %d vertex map max index is %d. There are only %d vertices in skeletal mesh, using %d vertices instead"), PartId, MaxIndex, NumVertices, NumVertices);
-				ReturnCode |= EUpdateResult::ERROR;
-			}
-		}
-		if (bIsVertexMapValid)
-		{
-			ClothPartData[PartId].VertexMap = AddConstant(PartVertexMap, StartIndex);
-			return (EUpdateResult)ReturnCode;
+			ClothPartData[PartId].VertexMap = PartVertexMap;
 		}
+		return (EUpdateResult)ReturnCode;
 	}
 
 	ClothPartData[PartId].VertexMap = Range(StartIndex, StartIndex + NumVertices);
