NearestNeighborModel: Smooth pca basis for high resolution cloth
Manually merge 23934139, 23942443 to avoid future conflicts

diff --git a/Engine/Plugins/Animation/MLDeformer/NearestNeighborModel/Source/NearestNeighborModel/Public/NearestNeighborModel.h b/Engine/Plugins/Animation/MLDeformer/NearestNeighborModel/Source/NearestNeighborModel/Public/NearestNeighborModel.h
--- a/Engine/Plugins/Animation/MLDeformer/NearestNeighborModel/Source/NearestNeighborModel/Public/NearestNeighborModel.h
+++ b/Engine/Plugins/Animation/MLDeformer/NearestNeighborModel/Source/NearestNeighborModel/Public/NearestNeighborModel.h
@@ -111,7 +111,7 @@ struct FClothPartData
 	UPROPERTY()
 	TArray<float> VertexMean;
 
-	/** PCA coefficients of the nearest neighbors. This is a flattened array of size (PCACoeffNum, NumNeighbors) */
+	/** PCA coefficients of the nearest neighbors. This is a flattened array of size (NumNeighbors, PCACoeffNum) */
 	UPROPERTY()
 	TArray<float> NeighborCoeffs;
 
@@ -206,7 +206,7 @@ public:
 
 	UFUNCTION(BlueprintCallable, Category = "Nearest Neighbor Model")
 	void SetNumNeighbors(int32 PartId, int32 InNumNeighbors) { ClothPartData[PartId].NumNeighbors = InNumNeighbors; }
-
+	
 	UFUNCTION(BlueprintPure, Category = "Nearest Neighbor Model")
 	const TArray<float>& NeighborCoeffs(int32 PartId) const { return ClothPartData[PartId].NeighborCoeffs; }
 
@@ -322,7 +322,6 @@ public:
 	static FName GetRecomputeDeltasPropertyName() { return GET_MEMBER_NAME_CHECKED(UNearestNeighborModel, bRecomputeDeltas); }
 	static FName GetRecomputePCAPropertyName() { return GET_MEMBER_NAME_CHECKED(UNearestNeighborModel, bRecomputePCA); }
 
-
 	bool GetUsePartOnlyMesh() const { return bUsePartOnlyMesh; }
 	static FName GetUsePartOnlyMeshPropertyName() { return GET_MEMBER_NAME_CHECKED(UNearestNeighborModel, bUsePartOnlyMesh); }
 
@@ -361,6 +360,9 @@ protected:
 	UPROPERTY(EditAnywhere, BlueprintReadWrite, Category = "Cloth Parts")
 	TArray<FClothPartEditorData> ClothPartEditorData;
 
+	UPROPERTY(EditAnywhere, BlueprintReadWrite, Category = "Cloth Parts", meta = (ClampMin = "0"))
+	int32 BasisSmoothIter = 0;
+
 	UPROPERTY()
 	bool bNearestNeighborDataValid = false;
 
@@ -390,6 +392,7 @@ protected:
 
 	UPROPERTY(EditAnywhere, BlueprintReadWrite, Category = "File Cache")
 	bool bRecomputePCA = true;
+
 #endif
 
 	static constexpr int32 NearestNeighborNumFloatsPerBone = 3;
diff --git a/Engine/Plugins/Animation/MLDeformer/NearestNeighborModel/Source/NearestNeighborModelEditor/Private/NearestNeighborGeomCacheSampler.cpp b/Engine/Plugins/Animation/MLDeformer/NearestNeighborModel/Source/NearestNeighborModelEditor/Private/NearestNeighborGeomCacheSampler.cpp
--- a/Engine/Plugins/Animation/MLDeformer/NearestNeighborModel/Source/NearestNeighborModelEditor/Private/NearestNeighborGeomCacheSampler.cpp
+++ b/Engine/Plugins/Animation/MLDeformer/NearestNeighborModel/Source/NearestNeighborModelEditor/Private/NearestNeighborGeomCacheSampler.cpp
@@ -356,5 +356,38 @@ namespace UE::NearestNeighborModel
 			}
 		}
 		return EUpdateResult::SUCCESS;
-	} 
+	}
+
+	TArray<uint32> FNearestNeighborGeomCacheSampler::GetMeshIndexBuffer() const
+	{
+		TArray<uint32> IndexBuffer;
+
+		if (SkeletalMeshComponent == nullptr)
+		{
+			return IndexBuffer;
+		}
+
+		USkeletalMesh* Mesh = SkeletalMeshComponent->GetSkeletalMeshAsset();
+		if (Mesh == nullptr)
+		{
+			return IndexBuffer;
+		}
+
+		constexpr int32 LODIndex = 0;
+		const FSkeletalMeshLODRenderData& SkelMeshLODData = Mesh->GetResourceForRendering()->LODRenderData[LODIndex];
+		SkelMeshLODData.MultiSizeIndexContainer.GetIndexBuffer(IndexBuffer);
+
+		const FSkeletalMeshModel* SkeletalMeshModel = Mesh->GetImportedModel();
+		const TArray<int32>& ImportedVertexNumbers = SkeletalMeshModel->LODModels[LODIndex].MeshToImportVertexMap;
+		
+		if (ImportedVertexNumbers.Num() > 0)
+		{
+			for (int32 Index = 0; Index < IndexBuffer.Num(); Index++)
+			{
+				IndexBuffer[Index] = ImportedVertexNumbers[IndexBuffer[Index]];
+			}
+		}
+
+		return IndexBuffer;
+	}
 };
diff --git a/Engine/Plugins/Animation/MLDeformer/NearestNeighborModel/Source/NearestNeighborModelEditor/Private/NearestNeighborTrainingModel.cpp b/Engine/Plugins/Animation/MLDeformer/NearestNeighborModel/Source/NearestNeighborModelEditor/Private/NearestNeighborTrainingModel.cpp
--- a/Engine/Plugins/Animation/MLDeformer/NearestNeighborModel/Source/NearestNeighborModelEditor/Private/NearestNeighborTrainingModel.cpp
+++ b/Engine/Plugins/Animation/MLDeformer/NearestNeighborModel/Source/NearestNeighborModelEditor/Private/NearestNeighborTrainingModel.cpp
@@ -90,5 +90,25 @@ int32 UNearestNeighborTrainingModel::GetKmeansNumClusters() const
 	return NearestNeighborModel->NumClusters;
 }
 
+const TArray<float> UNearestNeighborTrainingModel::GetUnskinnedVertexPositions() const
+{
+	const TArray<FVector3f>& PositionsVec = EditorModel->GetSampler()->GetUnskinnedVertexPositions();
+	TArray<float> PositionsFloat; 
+	PositionsFloat.SetNumUninitialized(PositionsVec.Num() * 3);
+	for (int32 i = 0; i < PositionsVec.Num(); i++)
+	{
+		PositionsFloat[i * 3] = PositionsVec[i].X;
+		PositionsFloat[i * 3 + 1] = PositionsVec[i].Y;
+		PositionsFloat[i * 3 + 2] = PositionsVec[i].Z;
+	}
+	return PositionsFloat;
+}
+
+const TArray<int32> UNearestNeighborTrainingModel::GetMeshIndexBuffer() const
+{
+	FNearestNeighborGeomCacheSampler* Sampler = static_cast<FNearestNeighborGeomCacheSampler*>(EditorModel->GetSampler());
+	TArray<uint32> UIndexBuffer = Sampler->GetMeshIndexBuffer();
+	return TArray<int32>((int32*)UIndexBuffer.GetData(), UIndexBuffer.Num());
+}
 
 #undef LOCTEXT_NAMESPACE
diff --git a/Engine/Plugins/Animation/MLDeformer/NearestNeighborModel/Source/NearestNeighborModelEditor/Public/NearestNeighborGeomCacheSampler.h b/Engine/Plugins/Animation/MLDeformer/NearestNeighborModel/Source/NearestNeighborModelEditor/Public/NearestNeighborGeomCacheSampler.h
--- a/Engine/Plugins/Animation/MLDeformer/NearestNeighborModel/Source/NearestNeighborModelEditor/Public/NearestNeighborGeomCacheSampler.h
+++ b/Engine/Plugins/Animation/MLDeformer/NearestNeighborModel/Source/NearestNeighborModelEditor/Public/NearestNeighborGeomCacheSampler.h
@@ -18,6 +18,7 @@ namespace UE::NearestNeighborModel
 		uint8 GenerateMeshMappingIndices();
 		uint8 CheckGeomCacheVertCount(int32 NumVertsFromGeomCache, int32 NumVertsFromVertexMap) const;
 		uint8 CheckMeshMappingsEmpty() const;
+		TArray<uint32> GetMeshIndexBuffer() const;
 
 	protected:
 		TArray<float> PartVertexDeltas;
diff --git a/Engine/Plugins/Animation/MLDeformer/NearestNeighborModel/Source/NearestNeighborModelEditor/Public/NearestNeighborTrainingModel.h b/Engine/Plugins/Animation/MLDeformer/NearestNeighborModel/Source/NearestNeighborModelEditor/Public/NearestNeighborTrainingModel.h
--- a/Engine/Plugins/Animation/MLDeformer/NearestNeighborModel/Source/NearestNeighborModelEditor/Public/NearestNeighborTrainingModel.h
+++ b/Engine/Plugins/Animation/MLDeformer/NearestNeighborModel/Source/NearestNeighborModelEditor/Public/NearestNeighborTrainingModel.h
@@ -64,6 +64,11 @@ public:
 	UFUNCTION(BlueprintPure, Category = "Python")
 	int32 GetKmeansNumClusters() const;
 
+	UFUNCTION(BlueprintPure, Category = "Python")
+	const TArray<float> GetUnskinnedVertexPositions() const;
+
+	UFUNCTION(BlueprintPure, Category = "Python")
+	const TArray<int32> GetMeshIndexBuffer() const;
 
 public:
 	UPROPERTY(EditAnywhere, BlueprintReadWrite, Category = "Training Data")
