NearestNeighborModel: add dual quaternion skinning

diff --git a/Engine/Plugins/Animation/MLDeformer/NearestNeighborModel/Source/NearestNeighborModel/Public/NearestNeighborModel.h b/Engine/Plugins/Animation/MLDeformer/NearestNeighborModel/Source/NearestNeighborModel/Public/NearestNeighborModel.h
--- a/Engine/Plugins/Animation/MLDeformer/NearestNeighborModel/Source/NearestNeighborModel/Public/NearestNeighborModel.h
+++ b/Engine/Plugins/Animation/MLDeformer/NearestNeighborModel/Source/NearestNeighborModel/Public/NearestNeighborModel.h
@@ -331,6 +331,8 @@ public:
 	int32 GetPartMeshIndex(int32 PartId) const { return ClothPartEditorData[PartId].MeshIndex; }
 	void SetPartMeshIndex(int32 PartId, int32 MeshIndex) { ClothPartEditorData[PartId].MeshIndex = MeshIndex; }
 	int32 GetMaxPartMeshIndex() const;
+
+	bool DoesUseDualQuaternionDeltas() const { return bUseDualQuaternionDeltas; }
 #endif
 
 
@@ -363,6 +365,9 @@ protected:
 	UPROPERTY(EditAnywhere, BlueprintReadWrite, Category = "Cloth Parts", meta = (ClampMin = "0"))
 	int32 BasisSmoothIter = 0;
 
+	UPROPERTY(EditAnywhere, BlueprintReadWrite, Category = "Cloth Parts")
+	bool bUseDualQuaternionDeltas = false;
+
 	UPROPERTY()
 	bool bNearestNeighborDataValid = false;
 
diff --git a/Engine/Plugins/Animation/MLDeformer/NearestNeighborModel/Source/NearestNeighborModelEditor/Private/NearestNeighborGeomCacheSampler.cpp b/Engine/Plugins/Animation/MLDeformer/NearestNeighborModel/Source/NearestNeighborModelEditor/Private/NearestNeighborGeomCacheSampler.cpp
--- a/Engine/Plugins/Animation/MLDeformer/NearestNeighborModel/Source/NearestNeighborModelEditor/Private/NearestNeighborGeomCacheSampler.cpp
+++ b/Engine/Plugins/Animation/MLDeformer/NearestNeighborModel/Source/NearestNeighborModelEditor/Private/NearestNeighborGeomCacheSampler.cpp
@@ -1,6 +1,7 @@
 // Copyright Epic Games, Inc. All Rights Reserved.
 
 #include "NearestNeighborGeomCacheSampler.h"
+#include "BoneWeights.h"
 #include "NearestNeighborModel.h"
 #include "Engine/SkeletalMesh.h"
 #include "Animation/DebugSkelMeshComponent.h"
@@ -17,6 +18,159 @@
 using namespace UE::MLDeformer;
 namespace UE::NearestNeighborModel
 {
+	void FNearestNeighborGeomCacheSampler::Sample(int32 InAnimFrameIndex)
+	{
+		UNearestNeighborModel* NearestNeighborModel = static_cast<UNearestNeighborModel*>(Model);
+		if (NearestNeighborModel->DoesUseDualQuaternionDeltas() && VertexDeltaSpace == EVertexDeltaSpace::PostSkinning)
+		{
+			FMLDeformerSampler::Sample(InAnimFrameIndex);
+			UpdateSkinnedPositions();
+			SampleDualQuaternionDeltas(InAnimFrameIndex);
+		}
+		else
+		{
+			FMLDeformerGeomCacheSampler::Sample(InAnimFrameIndex);
+		}
+	}
+
+	void FNearestNeighborGeomCacheSampler::SampleDualQuaternionDeltas(int32 InAnimFrameIndex)
+	{
+		USkeletalMesh* SkeletalMesh = SkeletalMeshComponent.Get() ? SkeletalMeshComponent->GetSkeletalMeshAsset() : nullptr;
+		UGeometryCache* GeometryCache = GeometryCacheComponent->GetGeometryCache();
+		if (SkeletalMeshComponent && SkeletalMesh && GeometryCacheComponent && GeometryCache)
+		{
+			const FTransform& AlignmentTransform = Model->GetAlignmentTransform();
+			FSkeletalMeshModel* ImportedModel = SkeletalMesh->GetImportedModel();
+
+			// For all mesh mappings we found.
+			const int32 LODIndex = 0;
+			const FSkeletalMeshLODModel& LODModel = ImportedModel->LODModels[LODIndex];
+			const TArray<FSkelMeshImportedMeshInfo>& SkelMeshInfos = LODModel.ImportedMeshInfos;
+			for (int32 MeshMappingIndex = 0; MeshMappingIndex < MeshMappings.Num(); ++MeshMappingIndex)
+			{
+				const UE::MLDeformer::FMLDeformerGeomCacheMeshMapping& MeshMapping = MeshMappings[MeshMappingIndex];
+				const FSkelMeshImportedMeshInfo& MeshInfo = SkelMeshInfos[MeshMapping.MeshIndex];
+				UGeometryCacheTrack* Track = GeometryCache->Tracks[MeshMapping.TrackIndex];
+
+				// Sample the mesh data of the geom cache.
+				FGeometryCacheMeshData& GeomCacheMeshData = GeomCacheMeshDatas[MeshMappingIndex];
+				if (!Track->GetMeshDataAtTime(SampleTime, GeomCacheMeshData))
+				{
+					continue;
+				}
+
+				// Calculate the vertex deltas.
+				const FSkeletalMeshLODRenderData& SkelMeshLODData = SkeletalMesh->GetResourceForRendering()->LODRenderData[LODIndex];
+				const FSkinWeightVertexBuffer& SkinWeightBuffer = *SkeletalMeshComponent->GetSkinWeightBuffer(LODIndex);
+				for (int32 VertexIndex = 0; VertexIndex < MeshInfo.NumVertices; ++VertexIndex)
+				{
+					const int32 SkinnedVertexIndex = MeshInfo.StartImportedVertex + VertexIndex;
+					const int32 GeomCacheVertexIndex = MeshMapping.SkelMeshToTrackVertexMap[VertexIndex];
+					if (GeomCacheVertexIndex != INDEX_NONE && GeomCacheMeshData.Positions.IsValidIndex(GeomCacheVertexIndex))
+					{
+						FVector3f Delta = FVector3f::ZeroVector;
+						const int32 ArrayIndex = 3 * SkinnedVertexIndex;
+						const int32 RenderVertexIndex = MeshMapping.ImportedVertexToRenderVertexMap[VertexIndex];
+						if (RenderVertexIndex != INDEX_NONE)
+						{
+							const FVector3f SkinnedVertexPos = SkinnedVertexPositions[SkinnedVertexIndex];
+							const FVector3f GeomCacheVertexPos = (FVector3f)AlignmentTransform.TransformPosition((FVector)GeomCacheMeshData.Positions[GeomCacheVertexIndex]);
+							const FVector3f WorldDelta = GeomCacheVertexPos - SkinnedVertexPos;
+							Delta = CalcDualQuaternionDelta(RenderVertexIndex, WorldDelta, SkelMeshLODData, SkinWeightBuffer);
+						}
+
+						VertexDeltas[ArrayIndex] = Delta.X;
+						VertexDeltas[ArrayIndex + 1] = Delta.Y;
+						VertexDeltas[ArrayIndex + 2] = Delta.Z;
+					}
+				}
+			}
+		}
+		else
+		{
+			VertexDeltas.Reset(0);
+		}
+	}
+
+	template<typename QuatType>
+	QuatType Conjugate(const QuatType& Q)
+	{
+		return QuatType(-Q.X, -Q.Y, -Q.Z, Q.W);
+	}
+
+	template<typename QuatType>
+	float Inner(const QuatType& Q1, const QuatType& Q2)
+	{
+		return Q1.X * Q2.X + Q1.Y * Q2.Y + Q1.Z * Q2.Z + Q1.W * Q2.W;
+	}
+
+	template<typename QuatType>
+	QuatType FromVector(const FVector3f& V)
+	{
+		return QuatType(V[0], V[1], V[2], 0);
+	}
+
+	template<typename QuatType>
+	FVector3f ToVector(const QuatType& Q)
+	{
+		return FVector3f(Q.X, Q.Y, Q.Z);
+	}
+
+	FVector3f FNearestNeighborGeomCacheSampler::CalcDualQuaternionDelta(int32 VertexIndex, const FVector3f& WorldDelta, const FSkeletalMeshLODRenderData& SkelMeshLODData, const FSkinWeightVertexBuffer& SkinWeightBuffer) const
+	{
+		check(SkeletalMeshComponent);
+		const USkeletalMesh* Mesh = SkeletalMeshComponent->GetSkeletalMeshAsset();
+		check(Mesh);
+
+		// Find the render section, which we need to find the right bone index.
+		int32 SectionIndex = INDEX_NONE;
+		int32 SectionVertexIndex = INDEX_NONE;
+		const FSkeletalMeshLODRenderData& LODData = Mesh->GetResourceForRendering()->LODRenderData[0];
+		LODData.GetSectionFromVertexIndex(VertexIndex, SectionIndex, SectionVertexIndex);
+
+		FQuat4f QuatSum = FQuat4f(0, 0, 0, 0);
+
+		FQuat4f R0;
+		float Sign = 1;
+		const int32 NumInfluences = SkinWeightBuffer.GetMaxBoneInfluences();
+		for (int32 InfluenceIndex = 0; InfluenceIndex < NumInfluences; ++InfluenceIndex)
+		{
+			const int32 BoneIndex = SkinWeightBuffer.GetBoneIndex(VertexIndex, InfluenceIndex);
+			const uint16 WeightByte = SkinWeightBuffer.GetBoneWeight(VertexIndex, InfluenceIndex);
+			// Weight must be > 0 when InflueceIndex == 0
+			ensure(InfluenceIndex > 0 || WeightByte > 0);
+			if (WeightByte > 0)
+			{
+				const int32 RealBoneIndex = LODData.RenderSections[SectionIndex].BoneMap[BoneIndex];
+				const FQuat4f R(BoneMatrices[RealBoneIndex]);
+				if (InfluenceIndex == 0)
+				{
+					R0 = R; Sign = 1;
+				}
+				else
+				{
+					Sign = Inner(R0, R) < 0 ? -1 : 1;
+				}
+				const float	Weight = Sign * static_cast<float>(WeightByte) * UE::AnimationCore::InvMaxRawBoneWeightFloat;
+				QuatSum += R * Weight;
+			}
+		}
+
+		const float SizeSquared = QuatSum.SizeSquared();
+		if (SizeSquared > SMALL_NUMBER)
+		{
+			/** Unrotate vector using v' = q^{-1} v q if q is unit size 
+			 * Because QuatSum is not unit size
+			 * v' =  q^* v q / |q|^2
+			 */
+			return ToVector<FQuat4f>(Conjugate(QuatSum) * FromVector<FQuat4f>(WorldDelta) * QuatSum) / SizeSquared;
+		}
+		else
+		{
+			return WorldDelta;			
+		}
+	}
+
 	uint8 FNearestNeighborGeomCacheSampler::SamplePart(int32 InAnimFrameIndex, int32 PartId)
 	{
 		FMLDeformerSampler::Sample(InAnimFrameIndex);
@@ -40,6 +194,7 @@ namespace UE::NearestNeighborModel
 				UE_LOG(LogNearestNeighborModel, Error, TEXT("SamplePart: MeshMappingIndices.Num()=%d is smaller than PartId %d"), MeshMappingIndices.Num(), PartId);
 				return EUpdateResult::ERROR;
 			}
+
 			const UE::MLDeformer::FMLDeformerGeomCacheMeshMapping& MeshMapping = MeshMappings[MeshMappingIndices[PartId]]; 
 
 			check(SkelMeshInfos.Num() > MeshMapping.MeshIndex);
@@ -64,6 +219,12 @@ namespace UE::NearestNeighborModel
 
 			UNearestNeighborModel* NearestNeighborModel = static_cast<UNearestNeighborModel*>(Model);
 			check(NearestNeighborModel != nullptr);
+
+			if (NearestNeighborModel->DoesUseDualQuaternionDeltas())
+			{
+				UpdateSkinnedPositions();
+			}
+
 			const TArray<uint32>& VertexMap = NearestNeighborModel->PartVertexMap(PartId);
 			const int32 NumPartVerts = VertexMap.Num();
 			PartVertexDeltas.Reset();
@@ -83,13 +244,22 @@ namespace UE::NearestNeighborModel
 					const int32 RenderVertexIndex = MeshMapping.ImportedVertexToRenderVertexMap[PartVertexIndex];
 					if (RenderVertexIndex != INDEX_NONE)
 					{
-						const FMatrix44f InvSkinningTransform = CalcInverseSkinningTransform(RenderVertexIndex, SkelMeshLODData, SkinWeightBuffer);
-
-						// Calculate the pre-skinning data.
-						const FVector3f UnskinnedPosition = SkelMeshLODData.StaticVertexBuffers.PositionVertexBuffer.VertexPosition(RenderVertexIndex);
 						const FVector3f GeomCacheVertexPos = (FVector3f)AlignmentTransform.TransformPosition((FVector)GeomCacheMeshData.Positions[GeomCacheVertexIndex]);
-						const FVector3f PreSkinningTargetPos = InvSkinningTransform.TransformPosition(GeomCacheVertexPos);
-						Delta = PreSkinningTargetPos - UnskinnedPosition;
+						if (NearestNeighborModel->DoesUseDualQuaternionDeltas())
+						{
+							const FVector3f SkinnedVertexPos = SkinnedVertexPositions[VertexIndex];
+							const FVector3f WorldDelta = GeomCacheVertexPos - SkinnedVertexPos;
+							Delta = CalcDualQuaternionDelta(RenderVertexIndex, WorldDelta, SkelMeshLODData, SkinWeightBuffer);
+						}
+						else
+						{
+							const FMatrix44f InvSkinningTransform = CalcInverseSkinningTransform(RenderVertexIndex, SkelMeshLODData, SkinWeightBuffer);
+
+							// Calculate the pre-skinning data.
+							const FVector3f UnskinnedPosition = SkelMeshLODData.StaticVertexBuffers.PositionVertexBuffer.VertexPosition(RenderVertexIndex);
+							const FVector3f PreSkinningTargetPos = InvSkinningTransform.TransformPosition(GeomCacheVertexPos);
+							Delta = PreSkinningTargetPos - UnskinnedPosition;
+						}
 					}
 
 					PartVertexDeltas[ArrayIndex] = Delta.X;
@@ -388,6 +558,6 @@ namespace UE::NearestNeighborModel
 			}
 		}
 
-		return IndexBuffer;
+		return MoveTemp(IndexBuffer);
 	}
 };
diff --git a/Engine/Plugins/Animation/MLDeformer/NearestNeighborModel/Source/NearestNeighborModelEditor/Private/NearestNeighborTrainingModel.cpp b/Engine/Plugins/Animation/MLDeformer/NearestNeighborModel/Source/NearestNeighborModelEditor/Private/NearestNeighborTrainingModel.cpp
--- a/Engine/Plugins/Animation/MLDeformer/NearestNeighborModel/Source/NearestNeighborModelEditor/Private/NearestNeighborTrainingModel.cpp
+++ b/Engine/Plugins/Animation/MLDeformer/NearestNeighborModel/Source/NearestNeighborModelEditor/Private/NearestNeighborTrainingModel.cpp
@@ -101,7 +101,7 @@ const TArray<float> UNearestNeighborTrainingModel::GetUnskinnedVertexPositions()
 		PositionsFloat[i * 3 + 1] = PositionsVec[i].Y;
 		PositionsFloat[i * 3 + 2] = PositionsVec[i].Z;
 	}
-	return PositionsFloat;
+	return MoveTemp(PositionsFloat);
 }
 
 const TArray<int32> UNearestNeighborTrainingModel::GetMeshIndexBuffer() const
diff --git a/Engine/Plugins/Animation/MLDeformer/NearestNeighborModel/Source/NearestNeighborModelEditor/Public/NearestNeighborGeomCacheSampler.h b/Engine/Plugins/Animation/MLDeformer/NearestNeighborModel/Source/NearestNeighborModelEditor/Public/NearestNeighborGeomCacheSampler.h
--- a/Engine/Plugins/Animation/MLDeformer/NearestNeighborModel/Source/NearestNeighborModelEditor/Public/NearestNeighborGeomCacheSampler.h
+++ b/Engine/Plugins/Animation/MLDeformer/NearestNeighborModel/Source/NearestNeighborModelEditor/Public/NearestNeighborGeomCacheSampler.h
@@ -10,6 +10,11 @@ namespace UE::NearestNeighborModel
 		: public UE::MLDeformer::FMLDeformerGeomCacheSampler
 	{
 	public:
+		// FMLDeformerGeomCacheSampler overrides
+		virtual void Sample(int32 InAnimFrameIndex) override;
+		// ~END FMLDeformerGeomCacheSampler overrides
+
+		virtual void SampleDualQuaternionDeltas(int32 InAnimFrameIndex);
 		virtual uint8 SamplePart(int32 InAnimFrameIndex, int32 PartId);
 		virtual bool SampleKMeansAnim(const int32 SkeletonId);
 		virtual bool SampleKMeansFrame(const int32 Frame);
@@ -19,6 +24,7 @@ namespace UE::NearestNeighborModel
 		uint8 CheckGeomCacheVertCount(int32 NumVertsFromGeomCache, int32 NumVertsFromVertexMap) const;
 		uint8 CheckMeshMappingsEmpty() const;
 		TArray<uint32> GetMeshIndexBuffer() const;
+		FVector3f CalcDualQuaternionDelta(int32 VertexIndex, const FVector3f& WorldDelta, const FSkeletalMeshLODRenderData& SkelMeshLODData, const FSkinWeightVertexBuffer& SkinWeightBuffer) const;
 
 	protected:
 		TArray<float> PartVertexDeltas;
