NearestNeighborModel: fix a typo from dual quaternion update

diff --git a/Engine/Plugins/Animation/MLDeformer/NearestNeighborModel/Source/NearestNeighborModelEditor/Private/NearestNeighborGeomCacheSampler.cpp b/Engine/Plugins/Animation/MLDeformer/NearestNeighborModel/Source/NearestNeighborModelEditor/Private/NearestNeighborGeomCacheSampler.cpp
--- a/Engine/Plugins/Animation/MLDeformer/NearestNeighborModel/Source/NearestNeighborModelEditor/Private/NearestNeighborGeomCacheSampler.cpp
+++ b/Engine/Plugins/Animation/MLDeformer/NearestNeighborModel/Source/NearestNeighborModelEditor/Private/NearestNeighborGeomCacheSampler.cpp
@@ -21,7 +21,7 @@ namespace UE::NearestNeighborModel
 	void FNearestNeighborGeomCacheSampler::Sample(int32 InAnimFrameIndex)
 	{
 		UNearestNeighborModel* NearestNeighborModel = static_cast<UNearestNeighborModel*>(Model);
-		if (NearestNeighborModel->DoesUseDualQuaternionDeltas() && VertexDeltaSpace == EVertexDeltaSpace::PostSkinning)
+		if (NearestNeighborModel->DoesUseDualQuaternionDeltas() && VertexDeltaSpace == EVertexDeltaSpace::PreSkinning)
 		{
 			FMLDeformerSampler::Sample(InAnimFrameIndex);
 			UpdateSkinnedPositions();
