NearestNeighborModel: fix vertex count mismatch bug. Reduce misleading error logs.

diff --git a/Engine/Plugins/Animation/MLDeformer/NearestNeighborModel/Source/NearestNeighborModel/Private/NearestNeighborModel.cpp b/Engine/Plugins/Animation/MLDeformer/NearestNeighborModel/Source/NearestNeighborModel/Private/NearestNeighborModel.cpp
--- a/Engine/Plugins/Animation/MLDeformer/NearestNeighborModel/Source/NearestNeighborModel/Private/NearestNeighborModel.cpp
+++ b/Engine/Plugins/Animation/MLDeformer/NearestNeighborModel/Source/NearestNeighborModel/Private/NearestNeighborModel.cpp
@@ -101,9 +101,12 @@ void UNearestNeighborModel::PostLoad()
 #endif
 }
 
-void UNearestNeighborModel::SetNNINetwork(UNeuralNetwork* InNeuralNetwork)
+void UNearestNeighborModel::SetNNINetwork(UNeuralNetwork* InNeuralNetwork, bool bBroadcast)
 {
-	GetNeuralNetworkModifyDelegate().Broadcast();
+	if (bBroadcast)
+	{
+		GetNeuralNetworkModifyDelegate().Broadcast();
+	}
 	NNINetwork = InNeuralNetwork;
 }
 
@@ -534,6 +537,8 @@ bool UNearestNeighborModel::LoadOptimizedNetwork(const FString& OnnxPath)
 		const bool bSuccess = Loader->LoadOptimizedNetwork(OnnxFile);
 		if (bSuccess)
 		{
+			// Clear optimized network to avoid error messages in SetOptimizedNetwork. 
+			OptimizedNetwork = nullptr;
 			SetOptimizedNetwork(Result);
 			return true;
 		}
diff --git a/Engine/Plugins/Animation/MLDeformer/NearestNeighborModel/Source/NearestNeighborModel/Public/NearestNeighborModel.h b/Engine/Plugins/Animation/MLDeformer/NearestNeighborModel/Source/NearestNeighborModel/Public/NearestNeighborModel.h
--- a/Engine/Plugins/Animation/MLDeformer/NearestNeighborModel/Source/NearestNeighborModel/Public/NearestNeighborModel.h
+++ b/Engine/Plugins/Animation/MLDeformer/NearestNeighborModel/Source/NearestNeighborModel/Public/NearestNeighborModel.h
@@ -246,7 +246,7 @@ public:
 	bool LoadOptimizedNetwork(const FString& OnnxPath);
 
 	UNeuralNetwork* GetNNINetwork() const;
-	void SetNNINetwork(UNeuralNetwork* InNeuralNetwork);
+	void SetNNINetwork(UNeuralNetwork* InNeuralNetwork, bool bBroadcast = true);
 
 #if WITH_EDITORONLY_DATA
 	TObjectPtr<UAnimSequence> GetNearestNeighborSkeletons(int32 PartId);
diff --git a/Engine/Plugins/Animation/MLDeformer/NearestNeighborModel/Source/NearestNeighborModelEditor/Private/NearestNeighborEditorModel.cpp b/Engine/Plugins/Animation/MLDeformer/NearestNeighborModel/Source/NearestNeighborModelEditor/Private/NearestNeighborEditorModel.cpp
--- a/Engine/Plugins/Animation/MLDeformer/NearestNeighborModel/Source/NearestNeighborModelEditor/Private/NearestNeighborEditorModel.cpp
+++ b/Engine/Plugins/Animation/MLDeformer/NearestNeighborModel/Source/NearestNeighborModelEditor/Private/NearestNeighborEditorModel.cpp
@@ -84,7 +84,7 @@ namespace UE::NearestNeighborModel
 
 	void FNearestNeighborEditorModel::OnPostTraining(ETrainingResult TrainingResult, bool bUsePartiallyTrainedWhenAborted)
 	{
-		FMLDeformerMorphModelEditorModel::OnPostTraining(TrainingResult, bUsePartiallyTrainedWhenAborted);
+		FMLDeformerEditorModel::OnPostTraining(TrainingResult, bUsePartiallyTrainedWhenAborted);
 		if (TrainingResult == ETrainingResult::Success || (TrainingResult == ETrainingResult::Aborted && bUsePartiallyTrainedWhenAborted))
 		{
 			InitTestMLDeformerPreviousWeights();
@@ -133,7 +133,7 @@ namespace UE::NearestNeighborModel
 				if (NearestNeighborModel->DoesEditorSupportOptimizedNetwork())
 				{
 					// We still need NNINetwork for serialization in nearest neighbor update and kmeans update.
-					NearestNeighborModel->SetNNINetwork(Network);
+					NearestNeighborModel->SetNNINetwork(Network, false);
 
 					const bool bSuccess = NearestNeighborModel->LoadOptimizedNetwork(OnnxFile);
 					if (bSuccess)
