NearestNeighborModel: optimize dual quaternion deltas performance. Allocate a quaternion buffer instead of converting from BoneMatrix to quaternions in vertex shader.

diff --git a/Engine/Plugins/Animation/MLDeformer/NearestNeighborModel/Shaders/Private/DataInterfaceSkeletonWithQuats.ush b/Engine/Plugins/Animation/MLDeformer/NearestNeighborModel/Shaders/Private/DataInterfaceSkeletonWithQuats.ush
new file mode 100644
--- /dev/null
+++ b/Engine/Plugins/Animation/MLDeformer/NearestNeighborModel/Shaders/Private/DataInterfaceSkeletonWithQuats.ush
@@ -0,0 +1,310 @@
+// Copyright Epic Games, Inc. All Rights Reserved.
+
+uint {DataInterfaceName}_NumVertices;
+uint {DataInterfaceName}_InputStreamStart;
+uint {DataInterfaceName}_NumBoneInfluences;
+uint {DataInterfaceName}_InputWeightStart;
+uint {DataInterfaceName}_InputWeightStride;
+uint {DataInterfaceName}_InputWeightIndexSize;
+Buffer<float4> {DataInterfaceName}_BoneMatrices;
+Buffer<float4> {DataInterfaceName}_BoneQuats;
+Buffer<uint> {DataInterfaceName}_InputWeightStream;
+Buffer<uint> {DataInterfaceName}_InputWeightLookupStream;
+
+uint ReadNumVertices_{DataInterfaceName}()
+{
+	return {DataInterfaceName}_NumVertices;
+}
+
+#if !GPUSKIN_UNLIMITED_BONE_INFLUENCE
+
+int4 GetBlendIndices_{DataInterfaceName}(uint StreamOffset, uint ExtraInfluenceIndex)
+{
+	int4 UnpackedBlendIndices = 0;
+
+#if GPUSKIN_BONE_INDEX_UINT16
+	StreamOffset += GPUSKIN_VB_OFFSET_INFLUENCEBONES + 2 * ExtraInfluenceIndex;
+	uint PackedBlendIndices = {DataInterfaceName}_InputWeightStream[StreamOffset];
+	UnpackedBlendIndices.x = PackedBlendIndices & 0xffff;
+	UnpackedBlendIndices.y = PackedBlendIndices >> 16 & 0xffff;
+	PackedBlendIndices = {DataInterfaceName}_InputWeightStream[StreamOffset + 1];
+	UnpackedBlendIndices.z = PackedBlendIndices & 0xffff;
+	UnpackedBlendIndices.w = PackedBlendIndices >> 16 & 0xffff;
+#else // GPUSKIN_BONE_INDEX_UINT16
+	StreamOffset += GPUSKIN_VB_OFFSET_INFLUENCEBONES + ExtraInfluenceIndex;
+	uint PackedBlendIndices = {DataInterfaceName}_InputWeightStream[StreamOffset];
+	UnpackedBlendIndices.x = PackedBlendIndices & 0xff;
+	UnpackedBlendIndices.y = PackedBlendIndices >> 8 & 0xff;
+	UnpackedBlendIndices.z = PackedBlendIndices >> 16 & 0xff;
+	UnpackedBlendIndices.w = PackedBlendIndices >> 24 & 0xff;
+#endif // GPUSKIN_BONE_INDEX_UINT16
+
+	return UnpackedBlendIndices;
+}
+
+float4 GetBlendWeights_{DataInterfaceName}(uint StreamOffset, uint ExtraInfluenceIndex)
+{
+	StreamOffset += (({DataInterfaceName}_NumBoneInfluences + 3) / 4) * (1 + GPUSKIN_BONE_INDEX_UINT16);
+
+	float4 UnpackedBlendWeights = 0;
+#if GPUSKIN_BONE_WEIGHTS_UINT16
+	StreamOffset += ExtraInfluenceIndex * 2;
+	uint PackedBlendWeights = {DataInterfaceName}_InputWeightStream[StreamOffset];
+	UnpackedBlendWeights.x = float(PackedBlendWeights & 0xffff) / 65535.0f;
+	UnpackedBlendWeights.y = float(PackedBlendWeights >> 16 & 0xffff) / 65535.0f;
+	PackedBlendWeights = {DataInterfaceName}_InputWeightStream[StreamOffset + 1];
+	UnpackedBlendWeights.z = float(PackedBlendWeights & 0xffff) / 65535.0f;
+	UnpackedBlendWeights.w = float(PackedBlendWeights >> 16 & 0xffff) / 65535.0f;
+	
+#else
+	StreamOffset += ExtraInfluenceIndex;
+	uint PackedBlendWeights = {DataInterfaceName}_InputWeightStream[StreamOffset];
+	
+	UnpackedBlendWeights.x = float(PackedBlendWeights & 0xff) / 255.0f;
+	UnpackedBlendWeights.y = float(PackedBlendWeights >> 8 & 0xff) / 255.0f;
+	UnpackedBlendWeights.z = float(PackedBlendWeights >> 16 & 0xff) / 255.0f;
+	UnpackedBlendWeights.w = float(PackedBlendWeights >> 24 & 0xff) / 255.0f;
+#endif	
+	return UnpackedBlendWeights;
+}
+
+#endif // GPUSKIN_UNLIMITED_BONE_INFLUENCE
+
+float3x4 GetBoneMatrix_{DataInterfaceName}(uint BoneIndex)
+{
+	uint BufferIndex = BoneIndex * 3;
+	return float3x4({DataInterfaceName}_BoneMatrices[BufferIndex], {DataInterfaceName}_BoneMatrices[BufferIndex + 1], {DataInterfaceName}_BoneMatrices[BufferIndex + 2]);
+}
+
+uint ReadNumBones_{DataInterfaceName}(uint VertexIndex)
+{
+#if !ENABLE_DEFORMER_BONES
+	return 1;
+#elif GPUSKIN_UNLIMITED_BONE_INFLUENCE
+	uint BlendOffsetCount = {DataInterfaceName}_InputWeightLookupStream[VertexIndex];
+	return BlendOffsetCount & 0xff;
+#else // !GPUSKIN_UNLIMITED_BONE_INFLUENCE
+	return {DataInterfaceName}_NumBoneInfluences;
+#endif
+}
+
+float3x4 ReadBoneMatrix_{DataInterfaceName}(uint VertexIndex, uint BoneIndex)
+{
+#if !ENABLE_DEFORMER_BONES
+
+	return float3x4(1, 0, 0, 0, 0, 1, 0, 0, 0, 0, 1, 0);
+
+#elif GPUSKIN_UNLIMITED_BONE_INFLUENCE
+	
+	uint BlendOffsetCount = {DataInterfaceName}_InputWeightLookupStream[VertexIndex];
+	int StreamOffset = BlendOffsetCount >> 8;
+
+	int BoneIndexOffset = StreamOffset + ({DataInterfaceName}_InputWeightIndexSize * BoneIndex);
+	int BoneGlobalIndex = {DataInterfaceName}_InputWeightStream[BoneIndexOffset];
+	if ({DataInterfaceName}_InputWeightIndexSize > 1)
+	{
+		BoneGlobalIndex = {DataInterfaceName}_InputWeightStream[BoneIndexOffset + 1] << 8 | BoneGlobalIndex;
+	}
+	
+	return GetBoneMatrix_{DataInterfaceName}(BoneGlobalIndex);
+
+#else // !GPUSKIN_UNLIMITED_BONE_INFLUENCE
+	
+	uint StreamOffset = {DataInterfaceName}_InputWeightStart + (VertexIndex * ({DataInterfaceName}_InputWeightStride / 4));
+	int BlendIndex = GetBlendIndices_{DataInterfaceName}(StreamOffset, BoneIndex / 4)[BoneIndex & 0x3];
+	return GetBoneMatrix_{DataInterfaceName}(BlendIndex);
+
+#endif
+}
+
+float ReadBoneWeight_{DataInterfaceName}(uint VertexIndex, uint BoneIndex)
+{
+#if !ENABLE_DEFORMER_BONES
+
+	return 1;
+
+#elif GPUSKIN_UNLIMITED_BONE_INFLUENCE
+	
+	uint BlendOffsetCount = {DataInterfaceName}_InputWeightLookupStream[VertexIndex];
+	int NumBoneInfluencesLocal = BlendOffsetCount & 0xff;
+	int StreamOffset = BlendOffsetCount >> 8;
+	int BoneIndexSize = {DataInterfaceName}_InputWeightIndexSize & 0xff;
+	int BoneWeightSize = {DataInterfaceName}_InputWeightIndexSize >> 8;
+	int WeightsOffset = StreamOffset + (BoneIndexSize * NumBoneInfluencesLocal);
+	int BoneWeightOffset = WeightsOffset + BoneIndex * BoneWeightSize;
+	
+	if (BoneWeightSize > 1)
+	{
+		return float({DataInterfaceName}_InputWeightStream[BoneWeightOffset + 1] << 8 | {DataInterfaceName}_InputWeightStream[BoneWeightOffset]) / 65535.0;
+	}
+	else
+	{
+		return float({DataInterfaceName}_InputWeightStream[BoneWeightOffset]) / 255.0f;
+	}
+
+#else // !GPUSKIN_UNLIMITED_BONE_INFLUENCE
+	
+	uint StreamOffset = {DataInterfaceName}_InputWeightStart + (VertexIndex * ({DataInterfaceName}_InputWeightStride / 4));
+	float BlendWeight = GetBlendWeights_{DataInterfaceName}(StreamOffset, BoneIndex / 4)[BoneIndex & 0x3];
+	return BlendWeight;
+
+#endif
+}
+
+float3x4 ReadWeightedBoneMatrix_{DataInterfaceName}(uint VertexIndex)
+{
+#if !ENABLE_DEFORMER_BONES
+
+	return float3x4(1, 0, 0, 0, 0, 1, 0, 0, 0, 0, 1, 0);
+
+#elif GPUSKIN_UNLIMITED_BONE_INFLUENCE
+
+	float3x4 BlendMatrix = (float3x4)0;
+
+	uint BufferIndex = {DataInterfaceName}_InputStreamStart + VertexIndex;
+	uint BlendOffsetCount = {DataInterfaceName}_InputWeightLookupStream[BufferIndex];
+	int NumBoneInfluencesLocal = BlendOffsetCount & 0xff;
+	int IndicesOffset = BlendOffsetCount >> 8;
+	int BoneIndexSize = {DataInterfaceName}_InputWeightIndexSize & 0xff;
+	int BoneWeightSize = {DataInterfaceName}_InputWeightIndexSize >> 8;
+	int WeightsOffset = IndicesOffset + (BoneIndexSize * NumBoneInfluencesLocal);
+
+	for (int InfluenceIndex = 0; InfluenceIndex < NumBoneInfluencesLocal; InfluenceIndex++)
+	{
+		int BoneIndexOffset = IndicesOffset + (BoneIndexSize * InfluenceIndex);
+		int BoneIndex = {DataInterfaceName}_InputWeightStream[BoneIndexOffset];
+		if (BoneIndexSize > 1)
+		{
+			BoneIndex = {DataInterfaceName}_InputWeightStream[BoneIndexOffset + 1] << 8 | BoneIndex;
+		}
+		int BoneWeightOffset = WeightsOffset + (BoneWeightSize * InfluenceIndex);
+		float BoneWeight = float({ DataInterfaceName }_InputWeightStream[BoneWeightOffset]) / 255.0f;
+		if (BoneWeightSize > 1)
+		{
+			BoneWeight = float({DataInterfaceName}_InputWeightStream[BoneWeightOffset + 1] << 8 | {DataInterfaceName}_InputWeightStream[BoneWeightOffset]) / 65535.0;
+		}
+		BlendMatrix += BoneWeight * GetBoneMatrix_{DataInterfaceName}(BoneIndex);
+	}
+
+	return BlendMatrix;
+
+#else // !GPUSKIN_UNLIMITED_BONE_INFLUENCE
+
+	float3x4 BlendMatrix = (float3x4)0;
+
+	uint StreamOffset = {DataInterfaceName}_InputWeightStart + (VertexIndex * ({DataInterfaceName}_InputWeightStride / 4));
+
+	if ({DataInterfaceName}_NumBoneInfluences > 0)
+	{
+		int4 BlendIndices = GetBlendIndices_{DataInterfaceName}(StreamOffset, 0);
+		float4 BlendWeights = GetBlendWeights_{DataInterfaceName}(StreamOffset, 0);
+
+		BlendMatrix += BlendWeights.x * GetBoneMatrix_{DataInterfaceName}(BlendIndices.x);
+		BlendMatrix += BlendWeights.y * GetBoneMatrix_{DataInterfaceName}(BlendIndices.y);
+		BlendMatrix += BlendWeights.z * GetBoneMatrix_{DataInterfaceName}(BlendIndices.z);
+		BlendMatrix += BlendWeights.w * GetBoneMatrix_{DataInterfaceName}(BlendIndices.w);
+	}
+
+	// todo[CF]: make this a permutation when permutation support is working again?
+	if ({DataInterfaceName}_NumBoneInfluences > 4)
+	{
+		int4 BlendIndices = GetBlendIndices_{DataInterfaceName}(StreamOffset, 1);
+		float4 BlendWeights = GetBlendWeights_{DataInterfaceName}(StreamOffset, 1);
+
+		BlendMatrix += BlendWeights.x * GetBoneMatrix_{DataInterfaceName}(BlendIndices.x);
+		BlendMatrix += BlendWeights.y * GetBoneMatrix_{DataInterfaceName}(BlendIndices.y);
+		BlendMatrix += BlendWeights.z * GetBoneMatrix_{DataInterfaceName}(BlendIndices.z);
+		BlendMatrix += BlendWeights.w * GetBoneMatrix_{DataInterfaceName}(BlendIndices.w);
+	}
+
+	return BlendMatrix;
+
+#endif
+}
+
+float4 GetBoneQuat_{DataInterfaceName}(uint BoneIndex)
+{
+	return {DataInterfaceName}_BoneQuats[BoneIndex];
+}
+
+float4 GetBoneQuatToBlend_{DataInterfaceName}(uint BoneIndex, float4 BoneQuat0, float Weight)
+{
+	float4 BoneQuat = GetBoneQuat_{DataInterfaceName}(BoneIndex);
+	return dot(BoneQuat, BoneQuat0) >= 0 ? Weight * BoneQuat : -Weight * BoneQuat;
+}
+
+float4 ReadWeightedBoneQuat_{DataInterfaceName}(uint VertexIndex)
+{
+#if !ENABLE_DEFORMER_BONES
+	return float4(0, 0, 0, 1);
+#elif GPUSKIN_UNLIMITED_BONE_INFLUENCE
+	float4 BlendQuat = (float4)0;
+	float4 BoneQuat0 = (float4)0;
+
+	uint BufferIndex = {DataInterfaceName}_InputStreamStart + VertexIndex;
+	uint BlendOffsetCount = {DataInterfaceName}_InputWeightLookupStream[BufferIndex];
+	int NumBoneInfluencesLocal = BlendOffsetCount & 0xff;
+	int IndicesOffset = BlendOffsetCount >> 8;
+	int BoneIndexSize = {DataInterfaceName}_InputWeightIndexSize & 0xff;
+	int BoneWeightSize = {DataInterfaceName}_InputWeightIndexSize >> 8;
+	int WeightsOffset = IndicesOffset + (BoneIndexSize * NumBoneInfluencesLocal);
+
+	for (int InfluenceIndex = 0; InfluenceIndex < NumBoneInfluencesLocal; InfluenceIndex++)
+	{
+		int BoneIndexOffset = IndicesOffset + (BoneIndexSize * InfluenceIndex);
+		int BoneIndex = {DataInterfaceName}_InputWeightStream[BoneIndexOffset];
+		if (BoneIndexSize > 1)
+		{
+			BoneIndex = {DataInterfaceName}_InputWeightStream[BoneIndexOffset + 1] << 8 | BoneIndex;
+		}
+		int BoneWeightOffset = WeightsOffset + (BoneWeightSize * InfluenceIndex);
+		float BoneWeight = float({ DataInterfaceName }_InputWeightStream[BoneWeightOffset]) / 255.0f;
+		if (BoneWeightSize > 1)
+		{
+			BoneWeight = float({DataInterfaceName}_InputWeightStream[BoneWeightOffset + 1] << 8 | {DataInterfaceName}_InputWeightStream[BoneWeightOffset]) / 65535.0;
+		}
+
+		if (InfluenceIndex == 0)
+		{
+			BoneQuat0 = GetBoneQuat_{DataInterfaceName}(BoneIndex);
+			BlendQuat += BoneWeight * BoneQuat0;
+		}
+		else
+		{
+			BlendQuat += GetBoneQuatToBlend_{DataInterfaceName}(BoneIndex, BoneQuat0, BoneWeight);
+		}
+	}
+	return BlendQuat;
+#else
+	float4 BlendQuat = (float4)0;
+	float4 BoneQuat0 = (float4)0;
+
+	uint StreamOffset = {DataInterfaceName}_InputWeightStart + (VertexIndex * ({DataInterfaceName}_InputWeightStride / 4));
+
+	if ({DataInterfaceName}_NumBoneInfluences > 0)
+	{
+		int4 BlendIndices = GetBlendIndices_{DataInterfaceName}(StreamOffset, 0);
+		float4 BlendWeights = GetBlendWeights_{DataInterfaceName}(StreamOffset, 0);
+
+		BoneQuat0 = GetBoneQuat_{DataInterfaceName}(BlendIndices.x);
+		BlendQuat += BlendWeights.x * BoneQuat0;
+		BlendQuat += GetBoneQuatToBlend_{DataInterfaceName}(BlendIndices.y, BoneQuat0, BlendWeights.y);
+		BlendQuat += GetBoneQuatToBlend_{DataInterfaceName}(BlendIndices.z, BoneQuat0, BlendWeights.z);
+		BlendQuat += GetBoneQuatToBlend_{DataInterfaceName}(BlendIndices.w, BoneQuat0, BlendWeights.w);
+	}
+
+	// todo[CF]: make this a permutation when permutation support is working again?
+	if ({DataInterfaceName}_NumBoneInfluences > 4)
+	{
+		int4 BlendIndices = GetBlendIndices_{DataInterfaceName}(StreamOffset, 1);
+		float4 BlendWeights = GetBlendWeights_{DataInterfaceName}(StreamOffset, 1);
+
+		BlendQuat += GetBoneQuatToBlend_{DataInterfaceName}(BlendIndices.x, BoneQuat0, BlendWeights.x);
+		BlendQuat += GetBoneQuatToBlend_{DataInterfaceName}(BlendIndices.y, BoneQuat0, BlendWeights.y);
+		BlendQuat += GetBoneQuatToBlend_{DataInterfaceName}(BlendIndices.z, BoneQuat0, BlendWeights.z);
+		BlendQuat += GetBoneQuatToBlend_{DataInterfaceName}(BlendIndices.w, BoneQuat0, BlendWeights.w);
+	}
+
+	return BlendQuat;
+#endif
+}
diff --git a/Engine/Plugins/Animation/MLDeformer/NearestNeighborModel/Source/NearestNeighborModel/Private/NearestNeighborModel.cpp b/Engine/Plugins/Animation/MLDeformer/NearestNeighborModel/Source/NearestNeighborModel/Private/NearestNeighborModel.cpp
--- a/Engine/Plugins/Animation/MLDeformer/NearestNeighborModel/Source/NearestNeighborModel/Private/NearestNeighborModel.cpp
+++ b/Engine/Plugins/Animation/MLDeformer/NearestNeighborModel/Source/NearestNeighborModel/Private/NearestNeighborModel.cpp
@@ -8,6 +8,7 @@
 #include "MLDeformerComponent.h"
 #include "Engine/SkeletalMesh.h"
 #include "GeometryCache.h"
+#include "Interfaces/IPluginManager.h"
 #include "Misc/FileHelper.h"
 #include "Components/SkinnedMeshComponent.h"
 #include "Components/ExternalMorphSet.h"
@@ -27,6 +28,11 @@ namespace UE::NearestNeighborModel
 	class NEARESTNEIGHBORMODEL_API FNearestNeighborModelModule
 		: public IModuleInterface
 	{
+		void StartupModule()
+		{
+			FString PluginShaderDir = FPaths::Combine(IPluginManager::Get().FindPlugin(TEXT("NearestNeighborModel"))->GetBaseDir(), TEXT("Shaders"));
+			AddShaderSourceDirectoryMapping(TEXT("/Plugin/MLDeformer/NearestNeighborModel"), PluginShaderDir);
+		}
 	};
 }
 IMPLEMENT_MODULE(UE::NearestNeighborModel::FNearestNeighborModelModule, NearestNeighborModel)
diff --git a/Engine/Plugins/Animation/MLDeformer/NearestNeighborModel/Source/NearestNeighborModel/Private/OptimusDataInterfaceSkeletonWithQuats.cpp b/Engine/Plugins/Animation/MLDeformer/NearestNeighborModel/Source/NearestNeighborModel/Private/OptimusDataInterfaceSkeletonWithQuats.cpp
new file mode 100644
--- /dev/null
+++ b/Engine/Plugins/Animation/MLDeformer/NearestNeighborModel/Source/NearestNeighborModel/Private/OptimusDataInterfaceSkeletonWithQuats.cpp
@@ -0,0 +1,322 @@
+// Copyright Epic Games, Inc. All Rights Reserved.
+
+#include "OptimusDataInterfaceSkeletonWithQuats.h"
+
+#include "ComputeFramework/ComputeKernelPermutationVector.h"
+#include "ComputeFramework/ShaderParamTypeDefinition.h"
+#include "OptimusDataDomain.h"
+#include "RenderGraphBuilder.h"
+#include "Rendering/SkeletalMeshLODRenderData.h"
+#include "Rendering/SkeletalMeshRenderData.h"
+#include "ShaderCompilerCore.h"
+#include "ShaderParameterMetadataBuilder.h"
+#include "SkeletalMeshDeformerHelpers.h"
+#include "SkeletalRenderPublic.h"
+
+#include UE_INLINE_GENERATED_CPP_BY_NAME(OptimusDataInterfaceSkeletonWithQuats)
+
+FString UOptimusSkeletonWithQuatsDataInterface::GetDisplayName() const
+{
+	return TEXT("SkeletonWithQuats");
+}
+
+TArray<FOptimusCDIPinDefinition> UOptimusSkeletonWithQuatsDataInterface::GetPinDefinitions() const
+{
+	TArray<FOptimusCDIPinDefinition> Defs;
+	Defs.Add({"NumBones", "ReadNumBones", Optimus::DomainName::Vertex, "ReadNumVertices"});
+	Defs.Add({"BoneMatrix", "ReadBoneMatrix", {{Optimus::DomainName::Vertex, "ReadNumVertices"}, {Optimus::DomainName::Bone, "ReadNumBones"}}});
+	Defs.Add({"BoneWeight", "ReadBoneWeight", {{Optimus::DomainName::Vertex, "ReadNumVertices"}, {Optimus::DomainName::Bone, "ReadNumBones"}}});
+	Defs.Add({"WeightedBoneMatrix", "ReadWeightedBoneMatrix", Optimus::DomainName::Vertex, "ReadNumVertices" });
+	Defs.Add({"WeightedBoneQuat", "ReadWeightedBoneQuat", Optimus::DomainName::Vertex, "ReadNumVertices" });
+	return Defs;
+}
+
+
+TSubclassOf<UActorComponent> UOptimusSkeletonWithQuatsDataInterface::GetRequiredComponentClass() const
+{
+	return USkinnedMeshComponent::StaticClass();
+}
+
+
+void UOptimusSkeletonWithQuatsDataInterface::GetSupportedInputs(TArray<FShaderFunctionDefinition>& OutFunctions) const
+{
+	OutFunctions.AddDefaulted_GetRef()
+		.SetName(TEXT("ReadNumVertices"))
+		.AddReturnType(EShaderFundamentalType::Uint);
+
+	OutFunctions.AddDefaulted_GetRef()
+		.SetName(TEXT("ReadNumBones"))
+		.AddReturnType(EShaderFundamentalType::Uint)
+		.AddParam(EShaderFundamentalType::Uint);
+
+	OutFunctions.AddDefaulted_GetRef()
+		.SetName(TEXT("ReadBoneMatrix"))
+		.AddReturnType(EShaderFundamentalType::Float, 3, 4)
+		.AddParam(EShaderFundamentalType::Uint)
+		.AddParam(EShaderFundamentalType::Uint);
+
+	OutFunctions.AddDefaulted_GetRef()
+		.SetName(TEXT("ReadBoneWeight"))
+		.AddReturnType(EShaderFundamentalType::Float)
+		.AddParam(EShaderFundamentalType::Uint)
+		.AddParam(EShaderFundamentalType::Uint);
+
+	OutFunctions.AddDefaulted_GetRef()
+		.SetName(TEXT("ReadWeightedBoneMatrix"))
+		.AddReturnType(EShaderFundamentalType::Float, 3, 4)
+		.AddParam(EShaderFundamentalType::Uint);
+
+	OutFunctions.AddDefaulted_GetRef()
+		.SetName(TEXT("ReadWeightedBoneQuat"))
+		.AddReturnType(EShaderFundamentalType::Float, 4)
+		.AddParam(EShaderFundamentalType::Uint);	
+}
+
+BEGIN_SHADER_PARAMETER_STRUCT(FSkeletonWithQuatsDataInterfaceParameters, )
+	SHADER_PARAMETER(uint32, NumVertices)
+	SHADER_PARAMETER(uint32, InputStreamStart)
+	SHADER_PARAMETER(uint32, NumBoneInfluences)
+	SHADER_PARAMETER(uint32, InputWeightStart)
+	SHADER_PARAMETER(uint32, InputWeightStride)
+	SHADER_PARAMETER(uint32, InputWeightIndexSize)
+	SHADER_PARAMETER_SRV(Buffer<float4>, BoneMatrices)
+	SHADER_PARAMETER_RDG_BUFFER_SRV(Buffer<float4>, BoneQuats)
+	SHADER_PARAMETER_SRV(Buffer<uint>, InputWeightStream)
+	SHADER_PARAMETER_SRV(Buffer<uint>, InputWeightLookupStream)
+END_SHADER_PARAMETER_STRUCT()
+
+void UOptimusSkeletonWithQuatsDataInterface::GetShaderParameters(TCHAR const* UID, FShaderParametersMetadataBuilder& InOutBuilder, FShaderParametersMetadataAllocations& InOutAllocations) const
+{
+	InOutBuilder.AddNestedStruct<FSkeletonWithQuatsDataInterfaceParameters>(UID);
+}
+
+void UOptimusSkeletonWithQuatsDataInterface::GetPermutations(FComputeKernelPermutationVector& OutPermutationVector) const
+{
+	// Need to be able to support these permutations according to the skeletal mesh settings.
+	// todo[CF]: I think GPUSKIN_UNLIMITED_BONE_INFLUENCE and GPUSKIN_BONE_INDEX_UINT16/GPUSKIN_BONE_WEIGHTS_UINT16 are mutually exclusive. So we could save permutations here.
+	OutPermutationVector.AddPermutation(TEXT("ENABLE_DEFORMER_BONES"), 2);
+	OutPermutationVector.AddPermutation(TEXT("GPUSKIN_UNLIMITED_BONE_INFLUENCE"), 2);
+	OutPermutationVector.AddPermutation(TEXT("GPUSKIN_BONE_INDEX_UINT16"), 2);
+	OutPermutationVector.AddPermutation(TEXT("GPUSKIN_BONE_WEIGHTS_UINT16"), 2);
+}
+
+TCHAR const* UOptimusSkeletonWithQuatsDataInterface::TemplateFilePath = TEXT("/Plugin/MLDeformer/NearestNeighborModel/Private/DataInterfaceSkeletonWithQuats.ush");
+
+TCHAR const* UOptimusSkeletonWithQuatsDataInterface::GetShaderVirtualPath() const
+{
+	return TemplateFilePath;
+}
+
+void UOptimusSkeletonWithQuatsDataInterface::GetShaderHash(FString& InOutKey) const
+{
+	GetShaderFileHash(TemplateFilePath, EShaderPlatform::SP_PCD3D_SM5).AppendString(InOutKey);
+}
+
+void UOptimusSkeletonWithQuatsDataInterface::GetHLSL(FString& OutHLSL, FString const& InDataInterfaceName) const
+{
+	TMap<FString, FStringFormatArg> TemplateArgs =
+	{
+		{ TEXT("DataInterfaceName"), InDataInterfaceName },
+	};
+
+	FString TemplateFile;
+	LoadShaderSourceFile(TemplateFilePath, EShaderPlatform::SP_PCD3D_SM5, &TemplateFile, nullptr);
+	OutHLSL += FString::Format(*TemplateFile, TemplateArgs);
+}
+
+UComputeDataProvider* UOptimusSkeletonWithQuatsDataInterface::CreateDataProvider(TObjectPtr<UObject> InBinding, uint64 InInputMask, uint64 InOutputMask) const
+{
+	UOptimusSkeletonWithQuatsDataProvider* Provider = NewObject<UOptimusSkeletonWithQuatsDataProvider>();
+	Provider->SkinnedMesh = Cast<USkinnedMeshComponent>(InBinding);
+	return Provider;
+}
+
+
+FComputeDataProviderRenderProxy* UOptimusSkeletonWithQuatsDataProvider::GetRenderProxy()
+{
+	return new FOptimusSkeletonWithQuatsDataProviderProxy(SkinnedMesh);
+}
+
+FOptimusSkeletonWithQuatsDataProviderProxy::FOptimusSkeletonWithQuatsDataProviderProxy(USkinnedMeshComponent* SkinnedMeshComponent)
+{
+	SkeletalMeshObject = SkinnedMeshComponent != nullptr ? SkinnedMeshComponent->MeshObject : nullptr;
+	BoneRevisionNumber = SkinnedMeshComponent != nullptr ? SkinnedMeshComponent->GetBoneTransformRevisionNumber() : 0;
+}
+
+bool FOptimusSkeletonWithQuatsDataProviderProxy::IsValid(FValidationData const& InValidationData) const
+{
+	if (InValidationData.ParameterStructSize != sizeof(FParameters))
+	{
+		return false;
+	}
+	if (SkeletalMeshObject == nullptr)
+	{
+		return false;
+	}
+	if (SkeletalMeshObject->GetSkeletalMeshRenderData().LODRenderData[SkeletalMeshObject->GetLOD()].RenderSections.Num() != InValidationData.NumInvocations)
+	{
+		return false;
+	}
+
+	return true;
+}
+
+struct FSkeletonWithQuatsDataInterfacePermutationIds
+{
+	uint32 EnableDeformerBones = 0;
+	uint32 UnlimitedBoneInfluence = 0;
+	uint32 BoneIndexUint16 = 0;
+	uint32 BoneWeightsUint16 = 0;
+
+	FSkeletonWithQuatsDataInterfacePermutationIds(FComputeKernelPermutationVector const& PermutationVector)
+	{
+		{
+			static FString Name(TEXT("ENABLE_DEFORMER_BONES"));
+			static uint32 Hash = GetTypeHash(Name);
+			EnableDeformerBones = PermutationVector.GetPermutationBits(Name, Hash, 1);
+		}
+		{
+			static FString Name(TEXT("GPUSKIN_UNLIMITED_BONE_INFLUENCE"));
+			static uint32 Hash = GetTypeHash(Name);
+			UnlimitedBoneInfluence = PermutationVector.GetPermutationBits(Name, Hash, 1);
+		}
+		{
+			static FString Name(TEXT("GPUSKIN_BONE_INDEX_UINT16"));
+			static uint32 Hash = GetTypeHash(Name);
+			BoneIndexUint16 = PermutationVector.GetPermutationBits(Name, Hash, 1);
+		}
+		{
+			static FString Name(TEXT("GPUSKIN_BONE_WEIGHTS_UINT16"));
+			static uint32 Hash = GetTypeHash(Name);
+			BoneWeightsUint16 = PermutationVector.GetPermutationBits(Name, Hash, 1);
+		}
+	}
+};
+
+void FOptimusSkeletonWithQuatsDataProviderProxy::GatherPermutations(FPermutationData& InOutPermutationData) const
+{
+	const int32 LodIndex = SkeletalMeshObject->GetLOD();
+	FSkeletalMeshRenderData const& SkeletalMeshRenderData = SkeletalMeshObject->GetSkeletalMeshRenderData();
+	FSkeletalMeshLODRenderData const* LodRenderData = &SkeletalMeshRenderData.LODRenderData[LodIndex];
+
+	FSkeletonWithQuatsDataInterfacePermutationIds PermutationIds(InOutPermutationData.PermutationVector);
+	for (int32 InvocationIndex = 0; InvocationIndex < InOutPermutationData.NumInvocations; ++InvocationIndex)
+	{
+		FSkelMeshRenderSection const& RenderSection = LodRenderData->RenderSections[InvocationIndex];
+
+		const bool bPreviousFrame = false;
+		FRHIShaderResourceView* BoneBufferSRV = FSkeletalMeshDeformerHelpers::GetBoneBufferForReading(SkeletalMeshObject, LodIndex, InvocationIndex, bPreviousFrame);
+
+		FSkinWeightVertexBuffer const* WeightBuffer = LodRenderData->GetSkinWeightVertexBuffer();
+		check(WeightBuffer != nullptr);
+		FRHIShaderResourceView* SkinWeightBufferSRV = WeightBuffer->GetDataVertexBuffer()->GetSRV();
+		const bool bUnlimitedBoneInfluences = WeightBuffer->GetBoneInfluenceType() == GPUSkinBoneInfluenceType::UnlimitedBoneInfluence;
+		FRHIShaderResourceView* InputWeightLookupStreamSRV = bUnlimitedBoneInfluences ? WeightBuffer->GetLookupVertexBuffer()->GetSRV() : nullptr;
+		const bool bValidBones = (BoneBufferSRV != nullptr) && (SkinWeightBufferSRV != nullptr) && (!bUnlimitedBoneInfluences || InputWeightLookupStreamSRV != nullptr);
+		const bool bUse16BitBoneIndex = WeightBuffer->Use16BitBoneIndex();
+		const bool bUse16BitBoneWeights = WeightBuffer->Use16BitBoneWeight();
+
+		InOutPermutationData.PermutationIds[InvocationIndex] |= (bValidBones ? PermutationIds.EnableDeformerBones : 0);
+		InOutPermutationData.PermutationIds[InvocationIndex] |= (bUnlimitedBoneInfluences ? PermutationIds.UnlimitedBoneInfluence : 0);
+		InOutPermutationData.PermutationIds[InvocationIndex] |= (bUse16BitBoneIndex ? PermutationIds.BoneIndexUint16 : 0);
+		InOutPermutationData.PermutationIds[InvocationIndex] |= (bUse16BitBoneWeights ? PermutationIds.BoneWeightsUint16 : 0);
+	}
+}
+
+void FOptimusSkeletonWithQuatsDataProviderProxy::AllocateResources(FRDGBuilder& GraphBuilder)
+{
+	TArray<FQuat4f> RefToLocalQuats;
+	CacheRefToLocalQuats(RefToLocalQuats);
+
+	const int32 LodIndex = SkeletalMeshObject->GetLOD();
+	FSkeletalMeshRenderData const& SkeletalMeshRenderData = SkeletalMeshObject->GetSkeletalMeshRenderData();
+	FSkeletalMeshLODRenderData const* LodRenderData = &SkeletalMeshRenderData.LODRenderData[LodIndex];
+	const int32 NumSections = LodRenderData->RenderSections.Num();
+	PerSectionRefToLocalQuats.SetNum(NumSections);
+	PerSectionRefToLocalQuatsSRVs.SetNum(NumSections);
+	for (int32 SectionIndex = 0; SectionIndex < NumSections; SectionIndex++)
+	{
+		TArray<FQuat4f> &SectionRefToLocalQuats = PerSectionRefToLocalQuats[SectionIndex];
+		FRDGBufferSRVRef &SectionRefToLocalQuatsSRV = PerSectionRefToLocalQuatsSRVs[SectionIndex];
+		const FSkelMeshRenderSection& Section = LodRenderData->RenderSections[SectionIndex];
+		const uint32 LocalNumBones = Section.BoneMap.Num();
+
+		SectionRefToLocalQuats.SetNum(LocalNumBones);
+		for (uint32 BoneIndex = 0; BoneIndex < LocalNumBones; BoneIndex++)
+		{
+			const FBoneIndexType RefToLocalIndex = Section.BoneMap[BoneIndex];
+			check(RefToLocalQuats.IsValidIndex(RefToLocalIndex));
+			SectionRefToLocalQuats[BoneIndex] = RefToLocalQuats[RefToLocalIndex];
+		}
+
+		if (LocalNumBones <= 0)
+		{
+			SectionRefToLocalQuatsSRV = GraphBuilder.CreateSRV(GraphBuilder.RegisterExternalBuffer(GWhiteVertexBufferWithRDG->Buffer), PF_A32B32G32R32F);
+		}
+		else
+		{
+			FRDGBufferRef SectionRefToLocalQuatsBuffer = GraphBuilder.CreateBuffer(FRDGBufferDesc::CreateStructuredDesc(sizeof(FQuat4f), LocalNumBones), TEXT("SectionRefToLocalQuats"));
+			SectionRefToLocalQuatsSRV = GraphBuilder.CreateSRV(SectionRefToLocalQuatsBuffer);
+			GraphBuilder.QueueBufferUpload(SectionRefToLocalQuatsBuffer, SectionRefToLocalQuats.GetData(), LocalNumBones * sizeof(FQuat4f));
+		}
+	}
+}
+
+void FOptimusSkeletonWithQuatsDataProviderProxy::GatherDispatchData(FDispatchData const& InDispatchData)
+{
+	const int32 LodIndex = SkeletalMeshObject->GetLOD();
+	FSkeletalMeshRenderData const& SkeletalMeshRenderData = SkeletalMeshObject->GetSkeletalMeshRenderData();
+	FSkeletalMeshLODRenderData const* LodRenderData = &SkeletalMeshRenderData.LODRenderData[LodIndex];
+
+	FRHIShaderResourceView* NullSRVBinding = GWhiteVertexBufferWithSRV->ShaderResourceViewRHI.GetReference();
+
+	const TStridedView<FParameters> ParameterArray = MakeStridedParameterView<FParameters>(InDispatchData);
+	for (int32 InvocationIndex = 0; InvocationIndex < ParameterArray.Num(); ++InvocationIndex)
+	{
+		FSkelMeshRenderSection const& RenderSection = LodRenderData->RenderSections[InvocationIndex];
+
+		const bool bPreviousFrame = false;
+		FRHIShaderResourceView* BoneBufferSRV = FSkeletalMeshDeformerHelpers::GetBoneBufferForReading(SkeletalMeshObject, LodIndex, InvocationIndex, bPreviousFrame);
+
+		FSkinWeightVertexBuffer const* WeightBuffer = LodRenderData->GetSkinWeightVertexBuffer();
+		check(WeightBuffer != nullptr);
+		FRHIShaderResourceView* SkinWeightBufferSRV = WeightBuffer->GetDataVertexBuffer()->GetSRV();
+		const bool bUnlimitedBoneInfluences = WeightBuffer->GetBoneInfluenceType() == GPUSkinBoneInfluenceType::UnlimitedBoneInfluence;
+		FRHIShaderResourceView* InputWeightLookupStreamSRV = bUnlimitedBoneInfluences ? WeightBuffer->GetLookupVertexBuffer()->GetSRV() : nullptr;
+
+		FParameters& Parameters = ParameterArray[InvocationIndex];
+		Parameters.NumVertices = RenderSection.NumVertices;
+		Parameters.InputStreamStart = RenderSection.BaseVertexIndex;
+		Parameters.NumBoneInfluences = WeightBuffer->GetMaxBoneInfluences();
+		Parameters.InputWeightStart = (WeightBuffer->GetConstantInfluencesVertexStride() * RenderSection.GetVertexBufferIndex()) / sizeof(float);
+		Parameters.InputWeightStride = WeightBuffer->GetConstantInfluencesVertexStride();
+		Parameters.InputWeightIndexSize = WeightBuffer->GetBoneIndexByteSize() | (WeightBuffer->GetBoneWeightByteSize() << 8);
+		Parameters.BoneMatrices = BoneBufferSRV != nullptr ? BoneBufferSRV : NullSRVBinding;
+		Parameters.BoneQuats = PerSectionRefToLocalQuatsSRVs[InvocationIndex];
+		Parameters.InputWeightStream = SkinWeightBufferSRV != nullptr ? SkinWeightBufferSRV : NullSRVBinding;
+		Parameters.InputWeightLookupStream = InputWeightLookupStreamSRV != nullptr ? InputWeightLookupStreamSRV : NullSRVBinding;
+	}
+}
+
+
+void FOptimusSkeletonWithQuatsDataProviderProxy::CacheRefToLocalQuats(TArray<FQuat4f>& OutRefToLocalQuats) const
+{
+	OutRefToLocalQuats.Reset();
+	if (SkeletalMeshObject == nullptr)
+	{
+		return;
+	}
+
+	const TArray<FMatrix44f>& RefToLocalMatrices = SkeletalMeshObject->GetReferenceToLocalMatrices();
+	const int32 NumQuats = RefToLocalMatrices.Num();
+	OutRefToLocalQuats.SetNumUninitialized(NumQuats);
+
+	/** It is theoretically more efficient to compute q_ref ^ {-1} * q_local
+	 *  Since we don't have access to q_ref in SkeletalMeshObject, we convert simply matrices to quaternions 
+	 */
+	for (int32 QuatIdx = 0; QuatIdx < NumQuats; ++QuatIdx)
+	{
+		OutRefToLocalQuats[QuatIdx] = RefToLocalMatrices[QuatIdx].ToQuat();
+	}
+}
+
diff --git a/Engine/Plugins/Animation/MLDeformer/NearestNeighborModel/Source/NearestNeighborModel/Private/OptimusDataInterfaceSkeletonWithQuats.h b/Engine/Plugins/Animation/MLDeformer/NearestNeighborModel/Source/NearestNeighborModel/Private/OptimusDataInterfaceSkeletonWithQuats.h
new file mode 100644
--- /dev/null
+++ b/Engine/Plugins/Animation/MLDeformer/NearestNeighborModel/Source/NearestNeighborModel/Private/OptimusDataInterfaceSkeletonWithQuats.h
@@ -0,0 +1,77 @@
+// Copyright Epic Games, Inc. All Rights Reserved.
+
+#pragma once
+
+#include "OptimusComputeDataInterface.h"
+#include "ComputeFramework/ComputeDataProvider.h"
+
+#include "OptimusDataInterfaceSkeletonWithQuats.generated.h"
+
+class FSkeletalMeshObject;
+class FSkeletonWithQuatsDataInterfaceParameters;
+class USkinnedMeshComponent;
+
+/** Compute Framework Data Interface for skeletal data. */
+UCLASS(Category = ComputeFramework)
+class NEARESTNEIGHBORMODEL_API UOptimusSkeletonWithQuatsDataInterface : public UOptimusComputeDataInterface
+{
+	GENERATED_BODY()
+
+public:
+	//~ Begin UOptimusComputeDataInterface Interface
+	FString GetDisplayName() const override;
+	TArray<FOptimusCDIPinDefinition> GetPinDefinitions() const override;
+	TSubclassOf<UActorComponent> GetRequiredComponentClass() const override;
+	//~ End UOptimusComputeDataInterface Interface
+	
+	//~ Begin UComputeDataInterface Interface
+	TCHAR const* GetClassName() const override { return TEXT("SkeletonWithQuats"); }
+	void GetSupportedInputs(TArray<FShaderFunctionDefinition>& OutFunctions) const override;
+	void GetShaderParameters(TCHAR const* UID, FShaderParametersMetadataBuilder& InOutBuilder, FShaderParametersMetadataAllocations& InOutAllocations) const override;
+	void GetPermutations(FComputeKernelPermutationVector& OutPermutationVector) const override;
+	TCHAR const* GetShaderVirtualPath() const override;
+	void GetShaderHash(FString& InOutKey) const override;
+	void GetHLSL(FString& OutHLSL, FString const& InDataInterfaceName) const override;
+	UComputeDataProvider* CreateDataProvider(TObjectPtr<UObject> InBinding, uint64 InInputMask, uint64 InOutputMask) const override;
+	//~ End UComputeDataInterface Interface
+
+private:
+	static TCHAR const* TemplateFilePath;
+};
+
+/** Compute Framework Data Provider for reading skeletal mesh. */
+UCLASS(BlueprintType, editinlinenew, Category = ComputeFramework)
+class UOptimusSkeletonWithQuatsDataProvider : public UComputeDataProvider
+{
+	GENERATED_BODY()
+
+public:
+	UPROPERTY(EditAnywhere, BlueprintReadWrite, Category = Binding)
+	TObjectPtr<USkinnedMeshComponent> SkinnedMesh = nullptr;
+
+	//~ Begin UComputeDataProvider Interface
+	FComputeDataProviderRenderProxy* GetRenderProxy() override;
+	//~ End UComputeDataProvider Interface
+};
+
+class FOptimusSkeletonWithQuatsDataProviderProxy : public FComputeDataProviderRenderProxy
+{
+public:
+	FOptimusSkeletonWithQuatsDataProviderProxy(USkinnedMeshComponent* SkinnedMeshComponent);
+
+	//~ Begin FComputeDataProviderRenderProxy Interface
+	bool IsValid(FValidationData const& InValidationData) const override;
+	void GatherPermutations(FPermutationData& InOutPermutationData) const override;
+	void GatherDispatchData(FDispatchData const& InDispatchData) override;
+	void AllocateResources(FRDGBuilder& GraphBuilder) override;
+	//~ End FComputeDataProviderRenderProxy Interface
+
+private:
+	using FParameters = FSkeletonWithQuatsDataInterfaceParameters;
+	void CacheRefToLocalQuats(TArray<FQuat4f>& OutRefToLocalQuats) const;
+
+	FSkeletalMeshObject* SkeletalMeshObject;
+	uint32 BoneRevisionNumber = 0;
+	TArray<TArray<FQuat4f> > PerSectionRefToLocalQuats;
+	TArray<FRDGBufferSRVRef> PerSectionRefToLocalQuatsSRVs;
+};
