Fix incorrect buffer type and view format on SectionRefToLocalQuats, leading to incorrect rendering on some platforms

diff --git a/Engine/Plugins/Animation/MLDeformer/NearestNeighborModel/Source/NearestNeighborModel/Private/OptimusDataInterfaceSkeletonWithQuats.cpp b/Engine/Plugins/Animation/MLDeformer/NearestNeighborModel/Source/NearestNeighborModel/Private/OptimusDataInterfaceSkeletonWithQuats.cpp
--- a/Engine/Plugins/Animation/MLDeformer/NearestNeighborModel/Source/NearestNeighborModel/Private/OptimusDataInterfaceSkeletonWithQuats.cpp
+++ b/Engine/Plugins/Animation/MLDeformer/NearestNeighborModel/Source/NearestNeighborModel/Private/OptimusDataInterfaceSkeletonWithQuats.cpp
@@ -255,8 +255,8 @@ void FOptimusSkeletonWithQuatsDataProviderProxy::AllocateResources(FRDGBuilder&
 		}
 		else
 		{
-			FRDGBufferRef SectionRefToLocalQuatsBuffer = GraphBuilder.CreateBuffer(FRDGBufferDesc::CreateStructuredDesc(sizeof(FQuat4f), LocalNumBones), TEXT("SectionRefToLocalQuats"));
-			SectionRefToLocalQuatsSRV = GraphBuilder.CreateSRV(SectionRefToLocalQuatsBuffer);
+			FRDGBufferRef SectionRefToLocalQuatsBuffer = GraphBuilder.CreateBuffer(FRDGBufferDesc::CreateBufferDesc(sizeof(FQuat4f), LocalNumBones), TEXT("SectionRefToLocalQuats"));
+			SectionRefToLocalQuatsSRV = GraphBuilder.CreateSRV(SectionRefToLocalQuatsBuffer, PF_A32B32G32R32F);
 			GraphBuilder.QueueBufferUpload(SectionRefToLocalQuatsBuffer, SectionRefToLocalQuats.GetData(), LocalNumBones * sizeof(FQuat4f));
 		}
 	}
