NearestNeighborModel: make time filtering DeltaTime dependent to reduce cloth popping when frame rate is high. Small bug fixes

diff --git a/Engine/Plugins/Animation/MLDeformer/NearestNeighborModel/Source/NearestNeighborModel/Private/NearestNeighborModelInstance.cpp b/Engine/Plugins/Animation/MLDeformer/NearestNeighborModel/Source/NearestNeighborModel/Private/NearestNeighborModelInstance.cpp
--- a/Engine/Plugins/Animation/MLDeformer/NearestNeighborModel/Source/NearestNeighborModel/Private/NearestNeighborModelInstance.cpp
+++ b/Engine/Plugins/Animation/MLDeformer/NearestNeighborModel/Source/NearestNeighborModel/Private/NearestNeighborModelInstance.cpp
@@ -132,7 +132,29 @@ void UNearestNeighborModelInstance::Execute(float ModelWeight)
 	{
 		Super::Execute(ModelWeight);
 	}
-	RunNearestNeighborModel(ModelWeight);
+}
+
+void UNearestNeighborModelInstance::Tick(float DeltaTime, float ModelWeight)
+{
+	// Post init hasn't yet succeeded, try again.
+	// This could for example happen when you add an ML Deformer component, but your SkeletalMesh isn't setup yet, but later becomes valid.
+	if (Model && !HasPostInitialized())
+	{
+		PostMLDeformerComponentInit();
+	}
+
+	if (ModelWeight > 0.0001f && HasValidTransforms() && SetupInputs())
+	{
+		// Execute the model instance.
+		// For models using neural networks this will perform the inference, 
+		// calculate the network outputs and possibly use them, depending on how the model works.
+		Execute(ModelWeight);
+		RunNearestNeighborModel(DeltaTime, ModelWeight);
+	}
+	else
+	{
+		HandleZeroModelWeight();
+	}
 }
 
 TArray<float> ExtractArray(const TArray<float>& InArr, int32 Start, int32 End)
@@ -146,7 +168,12 @@ TArray<float> ExtractArray(const TArray<float>& InArr, int32 Start, int32 End)
 	return Result;
 }
 
-void UNearestNeighborModelInstance::RunNearestNeighborModel(float ModelWeight)
+constexpr float FastExpMinusOne(float X)
+{
+	return X + X * X / 2 + X * X * X / 6;
+}
+
+void UNearestNeighborModelInstance::RunNearestNeighborModel(float DeltaTime, float ModelWeight)
 {
 	UNearestNeighborModel* NearestNeighborModel = Cast<UNearestNeighborModel>(Model);
 	if (NearestNeighborModel == nullptr)
@@ -154,6 +181,15 @@ void UNearestNeighborModelInstance::RunNearestNeighborModel(float ModelWeight)
 		return;
 	}
 	
+	/** DecayCoeff = (e ^ {(DecayFactor - 1) * DeltaTime + Delta0} - 1) / (e ^ {Delta0} - 1)
+	 * DecayFactor is a user defined value. DecayCoeff is the actual coefficient used to compute decay. 
+	 * DecayCoeff is an exponential decay of DeltaTime.
+	 * When DecayFactor = 1, DecayCoeff = 1; when DecayFactor = 0 && DeltaTime = Delta0, DecayCoeff = 0
+	 */	
+	constexpr float Delta0 = 1.0f / 30;
+	const float Delta = (NearestNeighborModel->GetDecayFactor() - 1) * DeltaTime;
+	const float DecayCoeff = FastExpMinusOne(FMath::Max(Delta + Delta0, 0)) / FastExpMinusOne(Delta0);
+
 	// Grab the weight data for this morph set.
 	// This could potentially fail if we are applying this deformer to the wrong skeletal mesh component.
 	const int LOD = 0;	// For now we only support LOD 0, as we can't setup an ML Deformer per LOD yet.
@@ -183,7 +219,7 @@ void UNearestNeighborModelInstance::RunNearestNeighborModel(float ModelWeight)
 			for (int32 MorphIndex = 0; MorphIndex < NumNetworkWeights; ++MorphIndex)
 			{
 				const float W = OutputTensorData[MorphIndex] * ModelWeight;
-				UpdateWeight(WeightData->Weights, MorphIndex + 1, W);
+				UpdateWeightWithDecay(WeightData->Weights, MorphIndex + 1, W, DecayCoeff);
 			}
 
 			int32 NeighborOffset = NumNetworkWeights + 1;
@@ -201,7 +237,7 @@ void UNearestNeighborModelInstance::RunNearestNeighborModel(float ModelWeight)
 					const int32 Index = NeighborOffset + NeighborId;
 					if (Index < NumMorphTargets)
 					{
-						UpdateWeight(WeightData->Weights, Index, W);
+						UpdateWeightWithDecay(WeightData->Weights, Index, W, DecayCoeff);
 					}
 				}
 				NeighborOffset += NumNeighbors;
@@ -214,12 +250,10 @@ void UNearestNeighborModelInstance::RunNearestNeighborModel(float ModelWeight)
 	}
 }
 
-void UNearestNeighborModelInstance::UpdateWeight(TArray<float>& MorphWeights, int32 Index, float W)
+void UNearestNeighborModelInstance::UpdateWeightWithDecay(TArray<float>& MorphWeights, int32 Index, float W, float DecayCoeff)
 {
-	UNearestNeighborModel* NearestNeighborModel = static_cast<UNearestNeighborModel*>(Model);
-	const float DecayFactor = NearestNeighborModel->GetDecayFactor();
 	const float PreviousW = PreviousWeights[Index];
-	const float NewW = (1 - DecayFactor) * W + DecayFactor * PreviousW;
+	const float NewW = (1 - DecayCoeff) * W + DecayCoeff * PreviousW;
 	MorphWeights[Index] = NewW;
 	PreviousWeights[Index] = NewW; 
 }
diff --git a/Engine/Plugins/Animation/MLDeformer/NearestNeighborModel/Source/NearestNeighborModel/Public/NearestNeighborModelInstance.h b/Engine/Plugins/Animation/MLDeformer/NearestNeighborModel/Source/NearestNeighborModel/Public/NearestNeighborModelInstance.h
--- a/Engine/Plugins/Animation/MLDeformer/NearestNeighborModel/Source/NearestNeighborModel/Public/NearestNeighborModelInstance.h
+++ b/Engine/Plugins/Animation/MLDeformer/NearestNeighborModel/Source/NearestNeighborModel/Public/NearestNeighborModelInstance.h
@@ -19,6 +19,7 @@ public:
     virtual void Execute(float ModelWeight) override;
     virtual bool SetupInputs() override;
     virtual FString CheckCompatibility(USkeletalMeshComponent* InSkelMeshComponent, bool LogIssues=false) override;
+    virtual void Tick(float DeltaTime, float ModelWeight) override;
 
 #if WITH_EDITORONLY_DATA
     const TArray<uint32>& GetNearestNeighborIds() const { return NearestNeighborIds; }
@@ -36,9 +37,9 @@ protected:
     // ~END UMLDeformerModelInstance overrides
 
 private:
-    void RunNearestNeighborModel(float ModelWeight);
+    void RunNearestNeighborModel(float DeltaTime, float ModelWeight);
     int32 FindNearestNeighbor(const float* PCAData, int32 PartId);
-    void UpdateWeight(TArray<float>& MorphWeights, int32 Index, float W);
+    void UpdateWeightWithDecay(TArray<float>& MorphWeights, int32 Index, float W, float DecayCoeff);
 
 #if WITH_EDITORONLY_DATA
     TArray<uint32> NearestNeighborIds;
diff --git a/Engine/Plugins/Animation/MLDeformer/NearestNeighborModel/Source/NearestNeighborModelEditor/Private/NearestNeighborEditorModel.cpp b/Engine/Plugins/Animation/MLDeformer/NearestNeighborModel/Source/NearestNeighborModelEditor/Private/NearestNeighborEditorModel.cpp
--- a/Engine/Plugins/Animation/MLDeformer/NearestNeighborModel/Source/NearestNeighborModelEditor/Private/NearestNeighborEditorModel.cpp
+++ b/Engine/Plugins/Animation/MLDeformer/NearestNeighborModel/Source/NearestNeighborModelEditor/Private/NearestNeighborEditorModel.cpp
@@ -87,7 +87,9 @@ namespace UE::NearestNeighborModel
 		FMLDeformerEditorModel::OnPostTraining(TrainingResult, bUsePartiallyTrainedWhenAborted);
 		if (TrainingResult == ETrainingResult::Success || (TrainingResult == ETrainingResult::Aborted && bUsePartiallyTrainedWhenAborted))
 		{
+			GetNearestNeighborModel()->InvalidateNearestNeighborData();
 			InitTestMLDeformerPreviousWeights();
+			GetEditor()->GetModelDetailsView()->ForceRefresh();
 		}
 	}
 
diff --git a/Engine/Plugins/Animation/MLDeformer/NearestNeighborModel/Source/NearestNeighborModelEditor/Private/NearestNeighborModelDetails.cpp b/Engine/Plugins/Animation/MLDeformer/NearestNeighborModel/Source/NearestNeighborModelEditor/Private/NearestNeighborModelDetails.cpp
--- a/Engine/Plugins/Animation/MLDeformer/NearestNeighborModel/Source/NearestNeighborModelEditor/Private/NearestNeighborModelDetails.cpp
+++ b/Engine/Plugins/Animation/MLDeformer/NearestNeighborModel/Source/NearestNeighborModelEditor/Private/NearestNeighborModelDetails.cpp
@@ -201,7 +201,7 @@ namespace UE::NearestNeighborModel
 		const int32 MaxPartMeshIndex = NearestNeighborModel->GetMaxPartMeshIndex();
 		if (MaxPartMeshIndex >= SubMeshNames.Num())
 		{
-			UE_LOG(LogNearestNeighborModel, Error, TEXT("Nearest neighbor model was previously created with %d submeshes, but the current skeletal mesh has %d submeshes. Please use the original skeletal mesh or click update to overwrite existing data."), MaxPartMeshIndex, SubMeshNames.Num());
+			UE_LOG(LogNearestNeighborModel, Error, TEXT("Nearest neighbor model was previously created with %d submeshes, but the current skeletal mesh has %d submeshes. Please use the original skeletal mesh or click update to overwrite existing data."), MaxPartMeshIndex + 1, SubMeshNames.Num());
 
 			NearestNeighborModel->InvalidateClothPartData();
 			AddActionResultText(ClothPartCategoryBuilder, EUpdateResult::ERROR, TEXT("Loading"));
