NearestNeighborModel: prevent unexpected usage of APIs.
Will gradually release APIs to public.

diff --git a/Engine/Plugins/Animation/MLDeformer/NearestNeighborModel/Source/NearestNeighborModel/Public/NearestNeighborModel.h b/Engine/Plugins/Animation/MLDeformer/NearestNeighborModel/Source/NearestNeighborModel/Public/NearestNeighborModel.h
--- a/Engine/Plugins/Animation/MLDeformer/NearestNeighborModel/Source/NearestNeighborModel/Public/NearestNeighborModel.h
+++ b/Engine/Plugins/Animation/MLDeformer/NearestNeighborModel/Source/NearestNeighborModel/Public/NearestNeighborModel.h
@@ -18,7 +18,16 @@ class UNeuralNetwork;
 class UMLDeformerAsset;
 class USkeleton;
 class IPropertyHandle;
+class UNearestNeighborModelInstance;
 class UNearestNeighborOptimizedNetwork;
+class UNearestNeighborTrainingModel;
+
+namespace UE::NearestNeighborModel
+{
+	class FNearestNeighborEditorModel;
+	class FNearestNeighborGeomCacheSampler;
+	class FNearestNeighborModelDetails;
+};
 
 NEARESTNEIGHBORMODEL_API DECLARE_LOG_CATEGORY_EXTERN(LogNearestNeighborModel, Log, All);
 
@@ -148,16 +157,16 @@ namespace UE::NearestNeighborModel
  * The pca basis and the nearest neighbor data are compressed into morph targets.
  */
 UCLASS()
-class NEARESTNEIGHBORMODEL_API UNearestNeighborModel 
-	: public UMLDeformerMorphModel
+class NEARESTNEIGHBORMODEL_API UNearestNeighborModel final 
+	: public UMLDeformerMorphModel 
 {
 	GENERATED_BODY()
 
 public:
 	UNearestNeighborModel(const FObjectInitializer& ObjectInitializer);
-	virtual void PostLoad() override;
 
 	// UMLDeformerModel overrides.
+	virtual void PostLoad() override;
 	virtual void Serialize(FArchive& Archive) override;
 	virtual bool DoesSupportQualityLevels() const override { return false; }
 	virtual UMLDeformerModelInstance* CreateModelInstance(UMLDeformerComponent* Component) override;
@@ -165,6 +174,13 @@ public:
 	virtual FString GetDisplayName() const override { return "Nearest Neighbor Model"; }
 	// ~END UMLDeformerModel overrides.
 
+	friend class UNearestNeighborModelInstance;
+	friend class UNearestNeighborTrainingModel;
+	friend class UE::NearestNeighborModel::FNearestNeighborEditorModel;
+	friend class UE::NearestNeighborModel::FNearestNeighborGeomCacheSampler;
+	friend class UE::NearestNeighborModel::FNearestNeighborModelDetails;
+
+private:
 	UFUNCTION(BlueprintPure, Category = "Nearest Neighbor Model")
 	int32 GetNumParts() const { return ClothPartData.Num(); }
 
@@ -335,7 +351,6 @@ public:
 	bool DoesUseDualQuaternionDeltas() const { return bUseDualQuaternionDeltas; }
 #endif
 
-
 protected:
 #if WITH_EDITORONLY_DATA
 	UPROPERTY(VisibleAnywhere, BlueprintReadOnly, Category = "Training Settings")
@@ -403,7 +418,6 @@ protected:
 	static constexpr int32 NearestNeighborNumFloatsPerBone = 3;
 	static constexpr int32 NearestNeighborNumFloatsPerCurve = 0;
 	
-public:
 #if WITH_EDITORONLY_DATA
 	UPROPERTY(EditAnywhere, BlueprintReadWrite, Category = "KMeans Pose Generator")
 	TArray<TObjectPtr<UAnimSequence>> SourceAnims;
@@ -436,6 +450,7 @@ public:
 	UPROPERTY(EditAnywhere, BlueprintReadWrite, Category = "Nearest Neighbors", META = (ClampMin = "0", ClampMax = "1"))
 	float NearestNeighborOffsetWeight = 1.0f;
 
+private:
 	UPROPERTY()
 	TObjectPtr<UNearestNeighborOptimizedNetwork> OptimizedNetwork = nullptr;
 
diff --git a/Engine/Plugins/Animation/MLDeformer/NearestNeighborModel/Source/NearestNeighborModel/Public/NearestNeighborModelInputInfo.h b/Engine/Plugins/Animation/MLDeformer/NearestNeighborModel/Source/NearestNeighborModel/Public/NearestNeighborModelInputInfo.h
--- a/Engine/Plugins/Animation/MLDeformer/NearestNeighborModel/Source/NearestNeighborModel/Public/NearestNeighborModelInputInfo.h
+++ b/Engine/Plugins/Animation/MLDeformer/NearestNeighborModel/Source/NearestNeighborModel/Public/NearestNeighborModelInputInfo.h
@@ -5,6 +5,13 @@
 #include "MLDeformerInputInfo.h"
 #include "NearestNeighborModelInputInfo.generated.h"
 
+class UNearestNeighborModelInstance;
+class UNearestNeighborModel;
+namespace UE::NearestNeighborModel
+{
+    class FNearestNeighborEditorModel;
+};
+
 UCLASS()
 class NEARESTNEIGHBORMODEL_API UNearestNeighborModelInputInfo
     : public UMLDeformerInputInfo
@@ -17,6 +24,11 @@ public:
     virtual void ExtractBoneRotations(USkeletalMeshComponent* SkelMeshComponent, TArray<float>& OutRotations) const override;
     // ~END UMLDeformerInputInfo overrides
 
+    friend class UNearestNeighborModel;
+    friend class UNearestNeighborModelInstance;
+    friend class UE::NearestNeighborModel::FNearestNeighborEditorModel;
+
+private:
     void InitRefBoneRotations(USkeletalMesh* SkelMesh);
     void ComputeNetworkInput(const TArray<FQuat>& BoneRotations, float* OutputBuffer, int64 StartIndex = 0) const;
 
diff --git a/Engine/Plugins/Animation/MLDeformer/NearestNeighborModel/Source/NearestNeighborModel/Public/NearestNeighborModelInstance.h b/Engine/Plugins/Animation/MLDeformer/NearestNeighborModel/Source/NearestNeighborModel/Public/NearestNeighborModelInstance.h
--- a/Engine/Plugins/Animation/MLDeformer/NearestNeighborModel/Source/NearestNeighborModel/Public/NearestNeighborModelInstance.h
+++ b/Engine/Plugins/Animation/MLDeformer/NearestNeighborModel/Source/NearestNeighborModel/Public/NearestNeighborModelInstance.h
@@ -5,10 +5,17 @@
 #include "NeuralTensor.h"
 #include "NearestNeighborModelInstance.generated.h"
 
+class UNearestNeighborModelInputInfo;
+class UNearestNeighborModelInstance;
 class UNearestNeighborOptimizedNetworkInstance;
+namespace UE::NearestNeighborModel
+{
+    class FNearestNeighborEditorModel;
+    class FNearestNeighborEditorModelActor;
+};
 
 UCLASS()
-class NEARESTNEIGHBORMODEL_API UNearestNeighborModelInstance
+class NEARESTNEIGHBORMODEL_API UNearestNeighborModelInstance final
     : public UMLDeformerMorphModelInstance
 {
     GENERATED_BODY()
@@ -20,7 +27,15 @@ public:
     virtual bool SetupInputs() override;
     virtual FString CheckCompatibility(USkeletalMeshComponent* InSkelMeshComponent, bool LogIssues=false) override;
     virtual void Tick(float DeltaTime, float ModelWeight) override;
+    // ~END UMLDeformerModelInstance overrides
+
+    friend class UNearestNeighborModelInputInfo;
+    friend class UNearestNeighborModelInstance;
+    friend class UNearestNeighborOptimizedNetworkInstance;
+    friend class UE::NearestNeighborModel::FNearestNeighborEditorModel;
+    friend class UE::NearestNeighborModel::FNearestNeighborEditorModelActor;
 
+private:
 #if WITH_EDITORONLY_DATA
     const TArray<uint32>& GetNearestNeighborIds() const { return NearestNeighborIds; }
     uint32 NearestNeighborId(int32 PartId) const { return NearestNeighborIds[PartId]; }
@@ -32,11 +47,8 @@ public:
     void GetInputDataPointer(float*& OutInputData, int32& OutNumInputFloats) const;
     void GetOutputDataPointer(float*& OutOutputData, int32& OutNumOutputFloats) const;
 
-protected:
     virtual int64 SetBoneTransforms(float* OutputBuffer, int64 OutputBufferSize, int64 StartIndex) override;
-    // ~END UMLDeformerModelInstance overrides
 
-private:
     void RunNearestNeighborModel(float DeltaTime, float ModelWeight);
     int32 FindNearestNeighbor(const float* PCAData, int32 PartId);
     void UpdateWeightWithDecay(TArray<float>& MorphWeights, int32 Index, float W, float DecayCoeff);
diff --git a/Engine/Plugins/Animation/MLDeformer/NearestNeighborModel/Source/NearestNeighborModel/Public/NearestNeighborModelVizSettings.h b/Engine/Plugins/Animation/MLDeformer/NearestNeighborModel/Source/NearestNeighborModel/Public/NearestNeighborModelVizSettings.h
--- a/Engine/Plugins/Animation/MLDeformer/NearestNeighborModel/Source/NearestNeighborModel/Public/NearestNeighborModelVizSettings.h
+++ b/Engine/Plugins/Animation/MLDeformer/NearestNeighborModel/Source/NearestNeighborModel/Public/NearestNeighborModelVizSettings.h
@@ -7,6 +7,10 @@
 #include "NearestNeighborModelVizSettings.generated.h"
 
 class UGeometryCache;
+namespace UE::NearestNeighborModel
+{
+    class FNearestNeighborEditorModel;
+};
 
 /**
  * The vizualization settings specific to the the vertex delta model.
@@ -16,14 +20,12 @@ class NEARESTNEIGHBORMODEL_API UNearestNeighborModelVizSettings
 	: public UMLDeformerMorphModelVizSettings
 {
 	GENERATED_BODY()
-public:
 #if	WITH_EDITORONLY_DATA
-	void SetNearestNeighborActorsOffset(float InOffset) { NearestNeighborActorsOffset = InOffset; }
-	float GetNearestNeighborActorsOffset() const { return NearestNeighborActorsOffset; }
+public:
 	static FName GetNearestNeighborActorsOffsetPropertyName() { return GET_MEMBER_NAME_CHECKED(UNearestNeighborModelVizSettings, NearestNeighborActorsOffset); }
-
 	static FName GetNearestNeighborIdsPropertyName() { return GET_MEMBER_NAME_CHECKED(UNearestNeighborModelVizSettings, NearestNeighborIds); }
-	void SetNearestNeighborIds(const TArray<uint32>& InNearestNeighborIds) { NearestNeighborIds = InNearestNeighborIds; }
+	
+	friend class UE::NearestNeighborModel::FNearestNeighborEditorModel;
 
 protected:
 	UPROPERTY(EditAnywhere, Category = "Live Settings")
@@ -31,5 +33,10 @@ protected:
 
 	UPROPERTY(VisibleAnywhere, Category = "Live Settings")
 	TArray<uint32> NearestNeighborIds;
+
+private:
+	void SetNearestNeighborActorsOffset(float InOffset) { NearestNeighborActorsOffset = InOffset; }
+	float GetNearestNeighborActorsOffset() const { return NearestNeighborActorsOffset; }
+	void SetNearestNeighborIds(const TArray<uint32>& InNearestNeighborIds) { NearestNeighborIds = InNearestNeighborIds; }
 #endif
 };
diff --git a/Engine/Plugins/Animation/MLDeformer/NearestNeighborModel/Source/NearestNeighborModelEditor/Public/NearestNeighborEditorModel.h b/Engine/Plugins/Animation/MLDeformer/NearestNeighborModel/Source/NearestNeighborModelEditor/Public/NearestNeighborEditorModel.h
--- a/Engine/Plugins/Animation/MLDeformer/NearestNeighborModel/Source/NearestNeighborModelEditor/Public/NearestNeighborEditorModel.h
+++ b/Engine/Plugins/Animation/MLDeformer/NearestNeighborModel/Source/NearestNeighborModelEditor/Public/NearestNeighborEditorModel.h
@@ -13,6 +13,7 @@ class UMLDeformerComponent;
 class UMLDeformerOptimizedNetwork;
 class UNearestNeighborModel;
 class UNearestNeighborModelInstance;
+class UNearestNeighborTrainingModel;
 class UNearestNeighborModelVizSettings;
 
 namespace UE::NearestNeighborModel
@@ -20,9 +21,10 @@ namespace UE::NearestNeighborModel
 	using namespace UE::MLDeformer;
 
 	class FNearestNeighborEditorModelActor;
+	class FNearestNeighborModelDetails;
 	class FNearestNeighborModelSampler;
 
-	class NEARESTNEIGHBORMODELEDITOR_API FNearestNeighborEditorModel 
+	class NEARESTNEIGHBORMODELEDITOR_API FNearestNeighborEditorModel final 
 		: public UE::MLDeformer::FMLDeformerMorphModelEditorModel
 	{
 	public:
@@ -50,6 +52,11 @@ namespace UE::NearestNeighborModel
 		virtual int32 GetNumTrainingFrames() const override;
 		// ~END FMLDeformerGeomCacheEditorModel overrides.
 
+		friend class FNearestNeighborModelSampler;
+		friend class FNearestNeighborModelDetails;
+		friend class ::UNearestNeighborTrainingModel;
+
+	private:
 		// Some helpers that cast to this model's variants of some classes.
 		UNearestNeighborModel* GetNearestNeighborModel() const { return static_cast<UNearestNeighborModel*>(Model); }
 		UNearestNeighborModelVizSettings* GetNearestNeighborModelVizSettings() const;
@@ -83,7 +90,6 @@ namespace UE::NearestNeighborModel
 		uint8 WarnIfNetworkInvalid();
 		bool IsNeuralNetworkLoaded();
 
-	protected:
 		virtual void CreateNearestNeighborActors(UWorld* World, int32 StartIndex = 0);
 
 		template<class TrainingModelClass>
@@ -104,7 +110,6 @@ namespace UE::NearestNeighborModel
 		// Actors showing the nearest neighbors of the current frame
 		TArray<FNearestNeighborEditorModelActor*> NearestNeighborActors;
 
-	private:
 		UWorld* EditorWorld = nullptr;
 		uint8 MorphTargetUpdateResult = 0;
 		uint8 KMeansClusterResult = 0;
diff --git a/Engine/Plugins/Animation/MLDeformer/NearestNeighborModel/Source/NearestNeighborModelEditor/Public/NearestNeighborEditorModelActor.h b/Engine/Plugins/Animation/MLDeformer/NearestNeighborModel/Source/NearestNeighborModelEditor/Public/NearestNeighborEditorModelActor.h
--- a/Engine/Plugins/Animation/MLDeformer/NearestNeighborModel/Source/NearestNeighborModelEditor/Public/NearestNeighborEditorModelActor.h
+++ b/Engine/Plugins/Animation/MLDeformer/NearestNeighborModel/Source/NearestNeighborModelEditor/Public/NearestNeighborEditorModelActor.h
@@ -10,6 +10,8 @@ class UGeometryCacheComponent;
 
 namespace UE::NearestNeighborModel
 {
+	class FNearestNeighborEditorModel;
+
 	using namespace UE::MLDeformer;
 
 	enum : int32
@@ -23,13 +25,15 @@ namespace UE::NearestNeighborModel
 	public:
 		FNearestNeighborEditorModelActor(const FConstructSettings& Settings);
 
+		friend class FNearestNeighborEditorModel; 
+
+	private:
 		void SetGeometryCacheComponent(UGeometryCacheComponent* Component) { GeomCacheComponent = Component; }
 		UGeometryCacheComponent* GetGeometryCacheComponent() const { return GeomCacheComponent; }
 
 		void InitNearestNeighborActor(const int32 InPartId, const UMLDeformerComponent* InComponent);
 		void TickNearestNeighborActor();
 
-	protected:
 		int32 PartId = INDEX_NONE;
 		const UMLDeformerComponent* MLDeformerComponent = nullptr;
 	};
diff --git a/Engine/Plugins/Animation/MLDeformer/NearestNeighborModel/Source/NearestNeighborModelEditor/Public/NearestNeighborGeomCacheSampler.h b/Engine/Plugins/Animation/MLDeformer/NearestNeighborModel/Source/NearestNeighborModelEditor/Public/NearestNeighborGeomCacheSampler.h
--- a/Engine/Plugins/Animation/MLDeformer/NearestNeighborModel/Source/NearestNeighborModelEditor/Public/NearestNeighborGeomCacheSampler.h
+++ b/Engine/Plugins/Animation/MLDeformer/NearestNeighborModel/Source/NearestNeighborModelEditor/Public/NearestNeighborGeomCacheSampler.h
@@ -4,8 +4,11 @@
 
 #include "MLDeformerGeomCacheSampler.h"
 
+class UNearestNeighborTrainingModel;
 namespace UE::NearestNeighborModel
 {
+	class FNearestNeighborEditorModel;
+
 	class NEARESTNEIGHBORMODELEDITOR_API FNearestNeighborGeomCacheSampler
 		: public UE::MLDeformer::FMLDeformerGeomCacheSampler
 	{
@@ -14,6 +17,10 @@ namespace UE::NearestNeighborModel
 		virtual void Sample(int32 InAnimFrameIndex) override;
 		// ~END FMLDeformerGeomCacheSampler overrides
 
+		friend class FNearestNeighborEditorModel;
+		friend class ::UNearestNeighborTrainingModel;
+
+	private:
 		virtual void SampleDualQuaternionDeltas(int32 InAnimFrameIndex);
 		virtual uint8 SamplePart(int32 InAnimFrameIndex, int32 PartId);
 		virtual bool SampleKMeansAnim(const int32 SkeletonId);
@@ -26,7 +33,6 @@ namespace UE::NearestNeighborModel
 		TArray<uint32> GetMeshIndexBuffer() const;
 		FVector3f CalcDualQuaternionDelta(int32 VertexIndex, const FVector3f& WorldDelta, const FSkeletalMeshLODRenderData& SkelMeshLODData, const FSkinWeightVertexBuffer& SkinWeightBuffer) const;
 
-	protected:
 		TArray<float> PartVertexDeltas;
 		TArray<int32> MeshMappingIndices;
 		int32 KMeansAnimId;
diff --git a/Engine/Plugins/Animation/MLDeformer/NearestNeighborModel/Source/NearestNeighborModelEditor/Public/NearestNeighborModelDetails.h b/Engine/Plugins/Animation/MLDeformer/NearestNeighborModel/Source/NearestNeighborModelEditor/Public/NearestNeighborModelDetails.h
--- a/Engine/Plugins/Animation/MLDeformer/NearestNeighborModel/Source/NearestNeighborModelEditor/Public/NearestNeighborModelDetails.h
+++ b/Engine/Plugins/Animation/MLDeformer/NearestNeighborModel/Source/NearestNeighborModelEditor/Public/NearestNeighborModelDetails.h
@@ -23,16 +23,14 @@ namespace UE::NearestNeighborModel
 
 		// ILayoutDetails overrides.
 		virtual void CustomizeDetails(IDetailLayoutBuilder& DetailBuilder) override;
-		// ~END ILayoutDetails overrides.
-
 		virtual void CreateCategories() override;
+		// ~END ILayoutDetails overrides.
 
+		// FMLDeformerModelDetails overrides.
 		virtual bool UpdateMemberPointers(const TArray<TWeakObjectPtr<UObject>>& Objects) override;
-		virtual void AddActionResultText(IDetailCategoryBuilder* CategoryBuilder, uint8 Result, const FString& ActionName);
+		// ~END FMLDeformerModelDetails overrides.
 
-		void GenerateClothPartElementWidget(TSharedRef<IPropertyHandle> PropertyHandle, int32 ArrayIndex, IDetailChildrenBuilder& ChildrenBuilder);
-
-	protected:
+	private:
 		UNearestNeighborModel* NearestNeighborModel = nullptr;
 		FNearestNeighborEditorModel* NearestNeighborEditorModel = nullptr;
 
@@ -43,6 +41,9 @@ namespace UE::NearestNeighborModel
 
 		TArray<TSharedPtr<FString> > SubMeshNames;
 		TMap<TSharedPtr<FString>, int32> SubMeshNameMap;
+
+		void AddActionResultText(IDetailCategoryBuilder* CategoryBuilder, uint8 Result, const FString& ActionName);
+		void GenerateClothPartElementWidget(TSharedRef<IPropertyHandle> PropertyHandle, int32 ArrayIndex, IDetailChildrenBuilder& ChildrenBuilder);
 		void BuildSubMeshNames();
 		void SubMeshComboSelectionChanged(TSharedPtr<FString> InSelectedItem, ESelectInfo::Type SelectInfo, int32 ArrayIndex);
 	};
diff --git a/Engine/Plugins/Animation/MLDeformer/NearestNeighborModel/Source/NearestNeighborModelEditor/Public/NearestNeighborTrainingModel.h b/Engine/Plugins/Animation/MLDeformer/NearestNeighborModel/Source/NearestNeighborModelEditor/Public/NearestNeighborTrainingModel.h
--- a/Engine/Plugins/Animation/MLDeformer/NearestNeighborModel/Source/NearestNeighborModelEditor/Public/NearestNeighborTrainingModel.h
+++ b/Engine/Plugins/Animation/MLDeformer/NearestNeighborModel/Source/NearestNeighborModelEditor/Public/NearestNeighborTrainingModel.h
@@ -14,7 +14,7 @@ namespace UE::NearestNeighborModel
 class UNearestNeighborModel;
 
 UCLASS(Blueprintable)
-class NEARESTNEIGHBORMODELEDITOR_API UNearestNeighborTrainingModel
+class NEARESTNEIGHBORMODELEDITOR_API UNearestNeighborTrainingModel final
 	: public UMLDeformerTrainingModel
 {
 	GENERATED_BODY()
@@ -26,6 +26,22 @@ public:
 	UFUNCTION(BlueprintImplementableEvent, Category = "Python")
 	int32 Train() const;
 
+	friend class UE::NearestNeighborModel::FNearestNeighborEditorModel;
+
+protected:
+	UFUNCTION(BlueprintImplementableEvent, Category = "Python")
+	int32 UpdateNearestNeighborData();
+
+	UFUNCTION(BlueprintImplementableEvent, Category = "Python")
+	int32 KmeansClusterPoses(const int32 PartId) const;
+
+	UPROPERTY(EditAnywhere, BlueprintReadWrite, Category = "Training Data")
+	TArray<float> PartSampleDeltas;
+
+	UPROPERTY(EditAnywhere, BlueprintReadWrite, Category = "Training Data")
+	TArray<int32> KmeansResults;
+
+private:
 	UFUNCTION(BlueprintPure, Category = "Python")
 	UNearestNeighborModel* GetNearestNeighborModel() const;
 
@@ -37,18 +53,12 @@ public:
 	UFUNCTION(BlueprintCallable, Category = "Python")
 	int32 SamplePart(int32 PartId, int32 Index);
 
-	UFUNCTION(BlueprintImplementableEvent, Category = "Python")
-	int32 UpdateNearestNeighborData();
-
 	UFUNCTION(BlueprintCallable, Category = "Python")
 	int32 SetSamplerPartData(const int32 PartId);
 
 	UFUNCTION(BlueprintPure, Category = "Python")
 	int32 GetPartNumNeighbors(const int32 PartId) const;
 
-	UFUNCTION(BlueprintImplementableEvent, Category = "Python")
-	int32 KmeansClusterPoses(const int32 PartId) const;
-
 	UFUNCTION(BlueprintCallable, Category = "Python")
 	bool SampleKmeansAnim(const int32 SkeletonId);
 
@@ -70,12 +80,5 @@ public:
 	UFUNCTION(BlueprintPure, Category = "Python")
 	const TArray<int32> GetMeshIndexBuffer() const;
 
-public:
-	UPROPERTY(EditAnywhere, BlueprintReadWrite, Category = "Training Data")
-	TArray<float> PartSampleDeltas;
-
-	UPROPERTY(EditAnywhere, BlueprintReadWrite, Category = "Training Data")
-	TArray<int32> KmeansResults;
-
 	UNearestNeighborModel* NearestNeighborModel = nullptr;
 };
