NearestNeighborModel: fix static warnings

diff --git a/Engine/Plugins/Animation/MLDeformer/NearestNeighborModel/Source/NearestNeighborModel/Public/NearestNeighborModel.h b/Engine/Plugins/Animation/MLDeformer/NearestNeighborModel/Source/NearestNeighborModel/Public/NearestNeighborModel.h
--- a/Engine/Plugins/Animation/MLDeformer/NearestNeighborModel/Source/NearestNeighborModel/Public/NearestNeighborModel.h
+++ b/Engine/Plugins/Animation/MLDeformer/NearestNeighborModel/Source/NearestNeighborModel/Public/NearestNeighborModel.h
@@ -157,7 +157,7 @@ namespace UE::NearestNeighborModel
  * The pca basis and the nearest neighbor data are compressed into morph targets.
  */
 UCLASS()
-class NEARESTNEIGHBORMODEL_API UNearestNeighborModel final 
+class NEARESTNEIGHBORMODEL_API UNearestNeighborModel
 	: public UMLDeformerMorphModel 
 {
 	GENERATED_BODY()
diff --git a/Engine/Plugins/Animation/MLDeformer/NearestNeighborModel/Source/NearestNeighborModel/Public/NearestNeighborModelInstance.h b/Engine/Plugins/Animation/MLDeformer/NearestNeighborModel/Source/NearestNeighborModel/Public/NearestNeighborModelInstance.h
--- a/Engine/Plugins/Animation/MLDeformer/NearestNeighborModel/Source/NearestNeighborModel/Public/NearestNeighborModelInstance.h
+++ b/Engine/Plugins/Animation/MLDeformer/NearestNeighborModel/Source/NearestNeighborModel/Public/NearestNeighborModelInstance.h
@@ -15,7 +15,7 @@ namespace UE::NearestNeighborModel
 };
 
 UCLASS()
-class NEARESTNEIGHBORMODEL_API UNearestNeighborModelInstance final
+class NEARESTNEIGHBORMODEL_API UNearestNeighborModelInstance
     : public UMLDeformerMorphModelInstance
 {
     GENERATED_BODY()
diff --git a/Engine/Plugins/Animation/MLDeformer/NearestNeighborModel/Source/NearestNeighborModelEditor/Public/NearestNeighborEditorModel.h b/Engine/Plugins/Animation/MLDeformer/NearestNeighborModel/Source/NearestNeighborModelEditor/Public/NearestNeighborEditorModel.h
--- a/Engine/Plugins/Animation/MLDeformer/NearestNeighborModel/Source/NearestNeighborModelEditor/Public/NearestNeighborEditorModel.h
+++ b/Engine/Plugins/Animation/MLDeformer/NearestNeighborModel/Source/NearestNeighborModelEditor/Public/NearestNeighborEditorModel.h
@@ -24,7 +24,7 @@ namespace UE::NearestNeighborModel
 	class FNearestNeighborModelDetails;
 	class FNearestNeighborModelSampler;
 
-	class NEARESTNEIGHBORMODELEDITOR_API FNearestNeighborEditorModel final 
+	class NEARESTNEIGHBORMODELEDITOR_API FNearestNeighborEditorModel
 		: public UE::MLDeformer::FMLDeformerMorphModelEditorModel
 	{
 	public:
diff --git a/Engine/Plugins/Animation/MLDeformer/NearestNeighborModel/Source/NearestNeighborModelEditor/Public/NearestNeighborTrainingModel.h b/Engine/Plugins/Animation/MLDeformer/NearestNeighborModel/Source/NearestNeighborModelEditor/Public/NearestNeighborTrainingModel.h
--- a/Engine/Plugins/Animation/MLDeformer/NearestNeighborModel/Source/NearestNeighborModelEditor/Public/NearestNeighborTrainingModel.h
+++ b/Engine/Plugins/Animation/MLDeformer/NearestNeighborModel/Source/NearestNeighborModelEditor/Public/NearestNeighborTrainingModel.h
@@ -14,7 +14,7 @@ namespace UE::NearestNeighborModel
 class UNearestNeighborModel;
 
 UCLASS(Blueprintable)
-class NEARESTNEIGHBORMODELEDITOR_API UNearestNeighborTrainingModel final
+class NEARESTNEIGHBORMODELEDITOR_API UNearestNeighborTrainingModel
 	: public UMLDeformerTrainingModel
 {
 	GENERATED_BODY()
