NearestNeighborModel: fix dual quaternion in matrices with scaling

diff --git a/Engine/Plugins/Animation/MLDeformer/NearestNeighborModel/Source/NearestNeighborModel/Private/OptimusDataInterfaceSkeletonWithQuats.cpp b/Engine/Plugins/Animation/MLDeformer/NearestNeighborModel/Source/NearestNeighborModel/Private/OptimusDataInterfaceSkeletonWithQuats.cpp
--- a/Engine/Plugins/Animation/MLDeformer/NearestNeighborModel/Source/NearestNeighborModel/Private/OptimusDataInterfaceSkeletonWithQuats.cpp
+++ b/Engine/Plugins/Animation/MLDeformer/NearestNeighborModel/Source/NearestNeighborModel/Private/OptimusDataInterfaceSkeletonWithQuats.cpp
@@ -316,7 +316,7 @@ void FOptimusSkeletonWithQuatsDataProviderProxy::CacheRefToLocalQuats(TArray<FQu
 	 */
 	for (int32 QuatIdx = 0; QuatIdx < NumQuats; ++QuatIdx)
 	{
-		OutRefToLocalQuats[QuatIdx] = RefToLocalMatrices[QuatIdx].ToQuat();
+		OutRefToLocalQuats[QuatIdx] = RefToLocalMatrices[QuatIdx].GetMatrixWithoutScale().ToQuat();
 	}
 }
 
diff --git a/Engine/Plugins/Animation/MLDeformer/NearestNeighborModel/Source/NearestNeighborModelEditor/Private/NearestNeighborGeomCacheSampler.cpp b/Engine/Plugins/Animation/MLDeformer/NearestNeighborModel/Source/NearestNeighborModelEditor/Private/NearestNeighborGeomCacheSampler.cpp
--- a/Engine/Plugins/Animation/MLDeformer/NearestNeighborModel/Source/NearestNeighborModelEditor/Private/NearestNeighborGeomCacheSampler.cpp
+++ b/Engine/Plugins/Animation/MLDeformer/NearestNeighborModel/Source/NearestNeighborModelEditor/Private/NearestNeighborGeomCacheSampler.cpp
@@ -142,7 +142,7 @@ namespace UE::NearestNeighborModel
 			if (WeightByte > 0)
 			{
 				const int32 RealBoneIndex = LODData.RenderSections[SectionIndex].BoneMap[BoneIndex];
-				const FQuat4f R(BoneMatrices[RealBoneIndex]);
+				const FQuat4f R(BoneMatrices[RealBoneIndex].GetMatrixWithoutScale());
 				if (InfluenceIndex == 0)
 				{
 					R0 = R; Sign = 1;
