NearestNeighborModel: fix sunlit assets

diff --git a/Engine/Plugins/Animation/MLDeformer/NearestNeighborModel/NearestNeighborModel.uplugin b/Engine/Plugins/Animation/MLDeformer/NearestNeighborModel/NearestNeighborModel.uplugin
--- a/Engine/Plugins/Animation/MLDeformer/NearestNeighborModel/NearestNeighborModel.uplugin
+++ b/Engine/Plugins/Animation/MLDeformer/NearestNeighborModel/NearestNeighborModel.uplugin
@@ -35,6 +35,10 @@
 			"Name": "NeuralNetworkInference",
 			"Enabled": true
 		},
+		{
+			"Name": "ComputeFramework",
+			"Enabled": true
+		},
 		{
 			"Name": "DeformerGraph",
 			"Enabled": true
