NearestNeighborModel: set NNI device type CPU

diff --git a/Engine/Plugins/Animation/MLDeformer/NearestNeighborModel/Source/NearestNeighborModelEditor/Private/NearestNeighborEditorModel.cpp b/Engine/Plugins/Animation/MLDeformer/NearestNeighborModel/Source/NearestNeighborModelEditor/Private/NearestNeighborEditorModel.cpp
--- a/Engine/Plugins/Animation/MLDeformer/NearestNeighborModel/Source/NearestNeighborModelEditor/Private/NearestNeighborEditorModel.cpp
+++ b/Engine/Plugins/Animation/MLDeformer/NearestNeighborModel/Source/NearestNeighborModelEditor/Private/NearestNeighborEditorModel.cpp
@@ -102,10 +102,10 @@ namespace UE::NearestNeighborModel
 			UNeuralNetwork* Result = NewObject<UNeuralNetwork>(Model, UNeuralNetwork::StaticClass());		
 			if (Result->Load(OnnxFile))
 			{
-				Result->SetDeviceType(ENeuralDeviceType::GPU, ENeuralDeviceType::CPU, ENeuralDeviceType::GPU);	
-				if (Result->GetDeviceType() != ENeuralDeviceType::GPU || Result->GetOutputDeviceType() != ENeuralDeviceType::GPU || Result->GetInputDeviceType() != ENeuralDeviceType::CPU)
+				Result->SetDeviceType(ENeuralDeviceType::CPU, ENeuralDeviceType::CPU, ENeuralDeviceType::CPU);	
+				if (Result->GetDeviceType() != ENeuralDeviceType::CPU || Result->GetOutputDeviceType() != ENeuralDeviceType::CPU || Result->GetInputDeviceType() != ENeuralDeviceType::CPU)
 				{
-					UE_LOG(LogNearestNeighborModel, Error, TEXT("Neural net in ML Deformer '%s' cannot run on the GPU, it will not be active."), *Model->GetDeformerAsset()->GetName());
+					UE_LOG(LogNearestNeighborModel, Error, TEXT("Neural net in ML Deformer '%s' cannot run on the CPU, it will not be active."), *Model->GetDeformerAsset()->GetName());
 				}
 				UE_LOG(LogNearestNeighborModel, Display, TEXT("Successfully loaded Onnx file '%s'..."), *OnnxFile);
 				return Result;
