NearestNeighborModel: remove NNI dependency
OptmizedNetworkLoader: remove DoesMeetPrerequisites

diff --git a/Engine/Plugins/Animation/MLDeformer/NearestNeighborModel/NearestNeighborModel.uplugin b/Engine/Plugins/Animation/MLDeformer/NearestNeighborModel/NearestNeighborModel.uplugin
--- a/Engine/Plugins/Animation/MLDeformer/NearestNeighborModel/NearestNeighborModel.uplugin
+++ b/Engine/Plugins/Animation/MLDeformer/NearestNeighborModel/NearestNeighborModel.uplugin
@@ -14,16 +14,21 @@
 	"IsBetaVersion": false,
 	"IsExperimentalVersion": true,
 	"Installed": false,
+	"SupportedTargetPlatforms": [
+		"Win64", "Linux", "Mac"
+	],
 	"Modules": [
 		{
 			"Name": "NearestNeighborModel",
 			"Type": "Runtime",
-			"LoadingPhase": "Default"
+			"LoadingPhase": "Default",
+			"PlatformAllowList": [ "Win64", "Linux", "Mac" ]
 		},
 		{
 			"Name": "NearestNeighborModelEditor",
 			"Type": "Editor",
-			"LoadingPhase": "Default"
+			"LoadingPhase": "Default",
+			"PlatformAllowList": [ "Win64", "Linux", "Mac" ]
 		}
 	],
 	"Plugins": [
@@ -32,7 +37,7 @@
 			"Enabled": true
 		},
 		{
-			"Name": "NeuralNetworkInference",
+			"Name": "NNE",
 			"Enabled": true
 		},
 		{
diff --git a/Engine/Plugins/Animation/MLDeformer/NearestNeighborModel/Source/NearestNeighborModel/NearestNeighborModel.Build.cs b/Engine/Plugins/Animation/MLDeformer/NearestNeighborModel/Source/NearestNeighborModel/NearestNeighborModel.Build.cs
--- a/Engine/Plugins/Animation/MLDeformer/NearestNeighborModel/Source/NearestNeighborModel/NearestNeighborModel.Build.cs
+++ b/Engine/Plugins/Animation/MLDeformer/NearestNeighborModel/Source/NearestNeighborModel/NearestNeighborModel.Build.cs
@@ -38,7 +38,9 @@ public class NearestNeighborModel : ModuleRules
 				"CoreUObject",
 				"Engine",
 				"GeometryCache",
-				"NeuralNetworkInference",
+				"NNECore",
+				"NNERuntimeRDG",
+				"NNEUtils",
 				"OptimusCore",
 				"Projects",
 				"RenderCore",
diff --git a/Engine/Plugins/Animation/MLDeformer/NearestNeighborModel/Source/NearestNeighborModel/Private/NearestNeighborModel.cpp b/Engine/Plugins/Animation/MLDeformer/NearestNeighborModel/Source/NearestNeighborModel/Private/NearestNeighborModel.cpp
--- a/Engine/Plugins/Animation/MLDeformer/NearestNeighborModel/Source/NearestNeighborModel/Private/NearestNeighborModel.cpp
+++ b/Engine/Plugins/Animation/MLDeformer/NearestNeighborModel/Source/NearestNeighborModel/Private/NearestNeighborModel.cpp
@@ -16,8 +16,7 @@
 #include "Rendering/MorphTargetVertexInfoBuffers.h"
 #include "ShaderCore.h"
 #include "UObject/UObjectGlobals.h"
-#include "NeuralNetwork.h"
-
+#include "NNECoreModelData.h"
 #include UE_INLINE_GENERATED_CPP_BY_NAME(NearestNeighborModel)
 
 #define LOCTEXT_NAMESPACE "UNearestNeighborModel"
@@ -54,7 +53,6 @@ UNearestNeighborModel::UNearestNeighborModel(const FObjectInitializer& ObjectIni
 {
 #if WITH_EDITORONLY_DATA
 	SetVizSettings(ObjectInitializer.CreateEditorOnlyDefaultSubobject<UNearestNeighborModelVizSettings>(this, TEXT("VizSettings")));
-	bDoesMeetOptimizedNetworkPrerequisites = DoesMeetOptimizedNetworkPrerequisites();
 #endif
 }
 
@@ -67,6 +65,20 @@ UMLDeformerInputInfo* UNearestNeighborModel::CreateInputInfo()
 	return NearestNeighborModelInputInfo;
 }
 
+void UNearestNeighborModel::SetNNEModelData(TObjectPtr<UNNEModelData> ModelData, bool bBroadcast)
+{
+	if (bBroadcast)
+	{
+		GetNeuralNetworkModifyDelegate().Broadcast();
+	}
+	NNEModel = ModelData;
+}
+
+TObjectPtr<UNNEModelData> UNearestNeighborModel::GetNNEModelData() const
+{
+	return NNEModel;
+}
+
 UMLDeformerModelInstance* UNearestNeighborModel::CreateModelInstance(UMLDeformerComponent* Component)
 {
 	return NewObject<UNearestNeighborModelInstance>(Component);
@@ -79,16 +91,16 @@ void UNearestNeighborModel::Serialize(FArchive& Archive)
 	{
 		if (DoesUseOptimizedNetwork())
 		{
-			if (GetNNINetwork() != nullptr)
+			if (GetNNEModelData() != nullptr)
 			{
-				SetNNINetwork(nullptr);
+				NNEModel = nullptr;
 			}
 		}
 		else
 		{
 			if (GetOptimizedNetwork() != nullptr)
 			{
-				SetOptimizedNetwork(nullptr);
+				OptimizedNetwork = nullptr;
 			}
 		}
 	}
@@ -108,20 +120,6 @@ void UNearestNeighborModel::PostLoad()
 #endif
 }
 
-void UNearestNeighborModel::SetNNINetwork(UNeuralNetwork* InNeuralNetwork, bool bBroadcast)
-{
-	if (bBroadcast)
-	{
-		GetNeuralNetworkModifyDelegate().Broadcast();
-	}
-	NNINetwork = InNeuralNetwork;
-}
-
-UNeuralNetwork* UNearestNeighborModel::GetNNINetwork() const
-{ 
-	return NNINetwork.Get();
-}
-
 TArray<uint32> ReadTxt(const FString &Path)
 {
 	TArray<FString> StrArr;
@@ -419,11 +417,11 @@ int32 UNearestNeighborModel::GetNumNeighborsFromAnimSequence(int32 PartId) const
 
 void UNearestNeighborModel::UpdateNetworkSize()
 {
-	UNeuralNetwork* Network = GetNNINetwork();
-	if (Network != nullptr)
+	if (NNEModel)
 	{
-		const SIZE_T NumBytes = Network->GetResourceSizeBytes(EResourceSizeMode::EstimatedTotal);
-		SavedNetworkSize = (double)NumBytes / 1024 / 1024;
+		const FString RuntimeName = GetNNERuntimeName();
+		TConstArrayView<uint8> ArrayView = NNEModel->GetModelData(RuntimeName);
+		SavedNetworkSize = ArrayView.GetTypeSize() * ArrayView.Num(); 
 	}
 	else
 	{
@@ -457,10 +455,10 @@ FString UNearestNeighborModel::GetModelDir() const
 	}
 }
 
-bool UNearestNeighborModel::DoesEditorSupportOptimizedNetwork() const
+bool UNearestNeighborModel::ShouldUseOptimizedNetwork() const
 {
 #if NEARESTNEIGHBORMODEL_USE_ISPC
-	return bDoesMeetOptimizedNetworkPrerequisites && UE::NearestNeighborModel::bNearestNeighborModelUseOptimizedNetwork;
+	return UE::NearestNeighborModel::bNearestNeighborModelUseOptimizedNetwork;
 #else
 	return false;
 #endif
@@ -563,14 +561,4 @@ int32 UNearestNeighborModel::GetOptimizedNetworkNumOutputs() const
 	return OptimizedNetwork ? OptimizedNetwork->GetNumOutputs() : 0;
 }
 
-bool UNearestNeighborModel::DoesMeetOptimizedNetworkPrerequisites() const
-{
-	UNearestNeighborOptimizedNetworkLoader* Loader = GetDerivedCDO<UNearestNeighborOptimizedNetworkLoader>();
-	if (Loader != nullptr)
-	{
-		return Loader->DoesMeetPrerequisites();
-	}
-	return false;
-}
-
 #undef LOCTEXT_NAMESPACE
diff --git a/Engine/Plugins/Animation/MLDeformer/NearestNeighborModel/Source/NearestNeighborModel/Private/NearestNeighborModelInstance.cpp b/Engine/Plugins/Animation/MLDeformer/NearestNeighborModel/Source/NearestNeighborModel/Private/NearestNeighborModelInstance.cpp
--- a/Engine/Plugins/Animation/MLDeformer/NearestNeighborModel/Source/NearestNeighborModel/Private/NearestNeighborModelInstance.cpp
+++ b/Engine/Plugins/Animation/MLDeformer/NearestNeighborModel/Source/NearestNeighborModel/Private/NearestNeighborModelInstance.cpp
@@ -3,7 +3,8 @@
 #include "NearestNeighborModelInstance.h"
 #include "NearestNeighborModel.h"
 #include "NearestNeighborModelInputInfo.h"
-#include "NeuralNetwork.h"
+#include "NNECore.h"
+#include "NNECoreModelData.h"
 #include "NearestNeighborOptimizedNetwork.h"
 #include "MLDeformerAsset.h"
 #include "Components/ExternalMorphSet.h"
@@ -44,51 +45,48 @@ int64 UNearestNeighborModelInstance::SetBoneTransforms(float* OutputBuffer, int6
 bool UNearestNeighborModelInstance::SetupInputs()
 {
 	const UNearestNeighborModel* NearestNeighborModel = Cast<UNearestNeighborModel>(Model);
-	check(NearestNeighborModel);
-	if (NearestNeighborModel && NearestNeighborModel->DoesUseOptimizedNetwork())
+	if (NearestNeighborModel == nullptr ||
+		SkeletalMeshComponent == nullptr ||
+		SkeletalMeshComponent->GetSkeletalMeshAsset() == nullptr ||
+		!bIsCompatible)
+	{
+		return false;
+	}
+	if (NearestNeighborModel->DoesUseOptimizedNetwork())
 	{
-		if (NearestNeighborModel == nullptr ||
-			SkeletalMeshComponent == nullptr ||
-			SkeletalMeshComponent->GetSkeletalMeshAsset() == nullptr ||
-			!bIsCompatible)
-		{
-			return false;
-		}
-
 		const UNearestNeighborOptimizedNetwork* OptimizedNetwork = NearestNeighborModel->GetOptimizedNetwork();
 		if (!OptimizedNetwork || !OptimizedNetworkInstance)
 		{
 			return false;
 		}
-
-		int32 NumNeuralNetInputs = 0;
-		float* InputDataPointer = nullptr;
-		GetInputDataPointer(InputDataPointer, NumNeuralNetInputs);
-
-		// If the neural network expects a different number of inputs, do nothing.
-		const int32 NumFloatsPerBone = NearestNeighborModel->GetNumFloatsPerBone();
-		const int32 NumFloatsPerCurve = NearestNeighborModel->GetNumFloatsPerCurve();
-		const int64 NumDeformerAssetInputs = NearestNeighborModel->GetInputInfo()->CalcNumNeuralNetInputs(NumFloatsPerBone, NumFloatsPerCurve);
-		if (NumNeuralNetInputs != NumDeformerAssetInputs)
-		{
-			return false;
-		}
-
-		const int64 NumFloatsWritten = SetNeuralNetworkInputValues(InputDataPointer, NumNeuralNetInputs);
-		check(NumFloatsWritten == NumNeuralNetInputs);
 	}
 	else
 	{
-		const bool bSuccess = Super::SetupInputs();
-		if (!bSuccess)
+		if (!ModelCPU || NNEInputTensorBuffer.IsEmpty() || NNEOutputTensorBuffer.IsEmpty())
 		{
 			return false;
 		}
 	}
 
-	float* InputDataPointer = nullptr;
 	int32 NumNeuralNetInputs = 0;
+	float* InputDataPointer = nullptr;
 	GetInputDataPointer(InputDataPointer, NumNeuralNetInputs);
+	if (!InputDataPointer)
+	{
+		return false;
+	}
+
+	// If the neural network expects a different number of inputs, do nothing.
+	const int32 NumFloatsPerBone = NearestNeighborModel->GetNumFloatsPerBone();
+	const int32 NumFloatsPerCurve = NearestNeighborModel->GetNumFloatsPerCurve();
+	const int64 NumDeformerAssetInputs = NearestNeighborModel->GetInputInfo()->CalcNumNeuralNetInputs(NumFloatsPerBone, NumFloatsPerCurve);
+	if (NumNeuralNetInputs != NumDeformerAssetInputs)
+	{
+		return false;
+	}
+
+	const int64 NumFloatsWritten = SetNeuralNetworkInputValues(InputDataPointer, NumNeuralNetInputs);
+	check(NumFloatsWritten == NumNeuralNetInputs);
 	if (InputDataPointer)
 	{
 		// Clip network inputs based on training set min and max values.
@@ -108,6 +106,10 @@ void UNearestNeighborModelInstance::Init(USkeletalMeshComponent* SkelMeshCompone
 	{
 		InitOptimizedNetworkInstance();
 	}
+	else
+	{
+		CreateNNEModel(); 
+	}
 }
 
 void UNearestNeighborModelInstance::InitOptimizedNetworkInstance()
@@ -130,7 +132,7 @@ void UNearestNeighborModelInstance::Execute(float ModelWeight)
 	}
 	else
 	{
-		Super::Execute(ModelWeight);
+		RunNNE();
 	}
 }
 
@@ -297,7 +299,46 @@ void UNearestNeighborModelInstance::InitPreviousWeights()
 	}
 }
 
-void UNearestNeighborModelInstance::GetInputDataPointer(float*& OutInputData, int32& OutNumInputFloats) const
+void UNearestNeighborModelInstance::CreateNNEModel()
+{
+	UNearestNeighborModel* NearestNeighborModel = Cast<UNearestNeighborModel>(Model);
+	if (NearestNeighborModel)
+	{
+		TWeakInterfacePtr<INNERuntime> Runtime = UE::NNECore::GetRuntime<INNERuntime>(NearestNeighborModel->GetNNERuntimeName());
+		TWeakInterfacePtr<INNERuntimeCPU> RuntimeCPU = UE::NNECore::GetRuntime<INNERuntimeCPU>(NearestNeighborModel->GetNNERuntimeName());
+
+		if (!Runtime.IsValid())
+		{
+			UE_LOG(LogNNE, Error, TEXT("Can't get %s runtime."), *NearestNeighborModel->GetNNERuntimeName());
+			return;
+		}
+
+		TObjectPtr<UNNEModelData> ModelData = NearestNeighborModel->GetNNEModelData(); 
+		if (ModelData)
+		{
+			if (RuntimeCPU.IsValid())
+			{
+				// allocate tensor inputs and outputs
+				ModelCPU = RuntimeCPU->CreateModelCPU(ModelData);
+
+				// setup inputs
+				TConstArrayView<UE::NNECore::FTensorDesc> InputTensorDescs = ModelCPU->GetInputTensorDescs();
+				UE::NNECore::FTensorShape InputTensorShape = UE::NNECore::FTensorShape::MakeFromSymbolic(InputTensorDescs[0].GetShape());
+				ModelCPU->SetInputTensorShapes({ InputTensorShape });
+				check(InputTensorDescs[0].GetElemByteSize() == sizeof(float)); 
+				NNEInputTensorBuffer.SetNumUninitialized(InputTensorShape.Volume());
+
+				// setup outputs
+				TConstArrayView<UE::NNECore::FTensorDesc> OutputTensorDescs = ModelCPU->GetOutputTensorDescs();
+				UE::NNECore::FTensorShape OutputTensorShape = UE::NNECore::FTensorShape::MakeFromSymbolic(OutputTensorDescs[0].GetShape());
+				check(OutputTensorDescs[0].GetElemByteSize() == sizeof(float));
+				NNEOutputTensorBuffer.SetNumUninitialized(OutputTensorShape.Volume());
+			}
+		}
+	}
+}
+
+void UNearestNeighborModelInstance::GetInputDataPointer(float*& OutInputData, int32& OutNumInputFloats)
 {
 	const UNearestNeighborModel* NearestNeighborModel = Cast<UNearestNeighborModel>(Model);
 	check(NearestNeighborModel);
@@ -313,19 +354,17 @@ void UNearestNeighborModelInstance::GetInputDataPointer(float*& OutInputData, in
 	}
 	else
 	{
-		UNeuralNetwork* NeuralNetwork = NearestNeighborModel->GetNNINetwork();
-		if (NeuralNetwork)
+		if (ModelCPU)
 		{
-			OutInputData = static_cast<float*>(NeuralNetwork->GetInputDataPointerMutableForContext(NeuralNetworkInferenceHandle));
-			OutNumInputFloats = NeuralNetwork->GetInputTensorForContext(NeuralNetworkInferenceHandle).Num();
-			return;
+			OutInputData = NNEInputTensorBuffer.GetData();
+			OutNumInputFloats = NNEInputTensorBuffer.Num(); 
 		}
 	}
 	OutInputData = nullptr;
 	OutNumInputFloats = 0;
 }
 
-void UNearestNeighborModelInstance::GetOutputDataPointer(float*& OutOutputData, int32& OutNumOutputFloats) const
+void UNearestNeighborModelInstance::GetOutputDataPointer(float*& OutOutputData, int32& OutNumOutputFloats)
 {
 	const UNearestNeighborModel* NearestNeighborModel = Cast<UNearestNeighborModel>(Model);
 	check(NearestNeighborModel);
@@ -341,12 +380,10 @@ void UNearestNeighborModelInstance::GetOutputDataPointer(float*& OutOutputData,
 	}
 	else
 	{
-		UNeuralNetwork* NeuralNetwork = NearestNeighborModel->GetNNINetwork();
-		if (NeuralNetwork)
+		if (ModelCPU)
 		{
-			const FNeuralTensor& OutputTensor = NeuralNetwork->GetOutputTensorForContext(NeuralNetworkInferenceHandle);
-			OutNumOutputFloats = OutputTensor.Num();
-			OutOutputData = (float*)OutputTensor.GetDataCasted<float>();
+			OutOutputData = NNEOutputTensorBuffer.GetData();
+			OutNumOutputFloats = NNEOutputTensorBuffer.Num();
 			return;
 		}
 	}
@@ -392,10 +429,9 @@ FString UNearestNeighborModelInstance::CheckCompatibility(USkeletalMeshComponent
 	}
 	else
 	{
-		UNeuralNetwork* NeuralNetwork = NearestNeighborModel->GetNNINetwork();
-		if (NeuralNetwork && NeuralNetwork->IsLoaded())
+		if (ModelCPU)
 		{
-			NumNeuralNetInputs = NeuralNetwork->GetInputTensor().Num();
+			NumNeuralNetInputs = NNEInputTensorBuffer.Num();
 		}
 	}
 
@@ -412,3 +448,40 @@ FString UNearestNeighborModelInstance::CheckCompatibility(USkeletalMeshComponent
 
 	return ErrorText;
 }
+
+TArray<float> UNearestNeighborModelInstance::Eval(const TArray<float>& InputData)
+{
+	const UNearestNeighborModel* NearestNeighborModel = Cast<UNearestNeighborModel>(Model);
+	int32 NumNeuralNetInputs = 0;
+	float* InputDataPointer = nullptr;
+	GetInputDataPointer(InputDataPointer, NumNeuralNetInputs);
+	check(InputDataPointer != nullptr);
+	check(InputData.Num() == NumNeuralNetInputs);
+
+	for (int32 i = 0; i < NumNeuralNetInputs; ++i)
+	{
+		InputDataPointer[i] = InputData[i];
+	}
+	NearestNeighborModel->ClipInputs(InputDataPointer, NumNeuralNetInputs);
+
+	if (NearestNeighborModel->DoesUseOptimizedNetwork())
+	{
+		check(OptimizedNetworkInstance != nullptr);
+		OptimizedNetworkInstance->Run();
+		TArrayView<float> Outputs = OptimizedNetworkInstance->GetOutputs();
+		return TArray<float>(Outputs.GetData(), Outputs.Num());
+	}
+	else
+	{
+		check(ModelCPU != nullptr);
+		RunNNE();
+		return NNEOutputTensorBuffer;
+	}
+}
+
+void UNearestNeighborModelInstance::RunNNE()
+{
+	UE::NNECore::FTensorBindingCPU InputTensorBinding{NNEInputTensorBuffer.GetData(), NNEInputTensorBuffer.Num()};
+	UE::NNECore::FTensorBindingCPU OutputTensorBinding{NNEOutputTensorBuffer.GetData(), NNEOutputTensorBuffer.Num()};
+	ModelCPU->RunSync({InputTensorBinding}, {OutputTensorBinding});
+}
\ No newline at end of file
diff --git a/Engine/Plugins/Animation/MLDeformer/NearestNeighborModel/Source/NearestNeighborModel/Public/NearestNeighborModel.h b/Engine/Plugins/Animation/MLDeformer/NearestNeighborModel/Source/NearestNeighborModel/Public/NearestNeighborModel.h
--- a/Engine/Plugins/Animation/MLDeformer/NearestNeighborModel/Source/NearestNeighborModel/Public/NearestNeighborModel.h
+++ b/Engine/Plugins/Animation/MLDeformer/NearestNeighborModel/Source/NearestNeighborModel/Public/NearestNeighborModel.h
@@ -14,7 +14,7 @@ struct FSkelMeshImportedMeshInfo;
 class USkeletalMesh;
 class UGeometryCache;
 class UAnimSequence;
-class UNeuralNetwork;
+class UNNEModelData;
 class UMLDeformerAsset;
 class USkeleton;
 class IPropertyHandle;
@@ -255,14 +255,16 @@ private:
 	virtual int32 GetNumFloatsPerCurve() const override { return NearestNeighborNumFloatsPerCurve; }
 
 	bool DoesUseOptimizedNetwork() const;
-	bool DoesMeetOptimizedNetworkPrerequisites() const;
+	bool ShouldUseOptimizedNetwork() const;
 	UNearestNeighborOptimizedNetwork* GetOptimizedNetwork() const { return OptimizedNetwork.Get(); }
 	int32 GetOptimizedNetworkNumOutputs() const;
 	void SetOptimizedNetwork(UNearestNeighborOptimizedNetwork* InOptimizedNetwork);
 	bool LoadOptimizedNetwork(const FString& OnnxPath);
 
-	UNeuralNetwork* GetNNINetwork() const;
-	void SetNNINetwork(UNeuralNetwork* InNeuralNetwork, bool bBroadcast = true);
+	void SetNNEModelData(TObjectPtr<UNNEModelData> ModelData, bool bBroadcast);
+	TObjectPtr<UNNEModelData> GetNNEModelData() const;
+
+	const FString GetNNERuntimeName() const {return TEXT("NNERuntimeORTDml"); }
 
 #if WITH_EDITORONLY_DATA
 	TObjectPtr<UAnimSequence> GetNearestNeighborSkeletons(int32 PartId);
@@ -290,7 +292,6 @@ private:
 	bool IsClothPartDataValid() const { return bClothPartDataValid; }
 	bool IsNearestNeighborDataValid() const { return bNearestNeighborDataValid; }
 	bool IsMorphTargetDataValid() const { return bMorphTargetDataValid; }
-	bool DoesEditorSupportOptimizedNetwork() const;
 	void SetUseOptimizedNetwork(bool bInUseOptimizedNetwork);
 #endif
 
@@ -454,12 +455,9 @@ private:
 	UPROPERTY()
 	TObjectPtr<UNearestNeighborOptimizedNetwork> OptimizedNetwork = nullptr;
 
-	/** The NNI neural network. */
-	UPROPERTY()
-	TObjectPtr<UNeuralNetwork> NNINetwork;
-
+	/** The NNE neural network model. */
 	UPROPERTY()
-	bool bDoesMeetOptimizedNetworkPrerequisites = false;
+	TObjectPtr<UNNEModelData> NNEModel;
 
 	UPROPERTY()
 	bool bUseOptimizedNetwork = true;
diff --git a/Engine/Plugins/Animation/MLDeformer/NearestNeighborModel/Source/NearestNeighborModel/Public/NearestNeighborModelInstance.h b/Engine/Plugins/Animation/MLDeformer/NearestNeighborModel/Source/NearestNeighborModel/Public/NearestNeighborModelInstance.h
--- a/Engine/Plugins/Animation/MLDeformer/NearestNeighborModel/Source/NearestNeighborModel/Public/NearestNeighborModelInstance.h
+++ b/Engine/Plugins/Animation/MLDeformer/NearestNeighborModel/Source/NearestNeighborModel/Public/NearestNeighborModelInstance.h
@@ -2,9 +2,11 @@
 #pragma once
 
 #include "MLDeformerMorphModelInstance.h"
-#include "NeuralTensor.h"
+#include "NNECore.h"
+#include "NNECoreRuntimeCPU.h"
 #include "NearestNeighborModelInstance.generated.h"
 
+
 class UNearestNeighborModelInputInfo;
 class UNearestNeighborModelInstance;
 class UNearestNeighborOptimizedNetworkInstance;
@@ -44,8 +46,8 @@ private:
 
     void InitPreviousWeights();
     void InitOptimizedNetworkInstance();
-    void GetInputDataPointer(float*& OutInputData, int32& OutNumInputFloats) const;
-    void GetOutputDataPointer(float*& OutOutputData, int32& OutNumOutputFloats) const;
+    void GetInputDataPointer(float*& OutInputData, int32& OutNumInputFloats);
+    void GetOutputDataPointer(float*& OutOutputData, int32& OutNumOutputFloats);
 
     virtual int64 SetBoneTransforms(float* OutputBuffer, int64 OutputBufferSize, int64 StartIndex) override;
 
@@ -53,6 +55,9 @@ private:
     int32 FindNearestNeighbor(const float* PCAData, int32 PartId);
     void UpdateWeightWithDecay(TArray<float>& MorphWeights, int32 Index, float W, float DecayCoeff);
 
+    UFUNCTION(BlueprintCallable, Category = "Python")
+    TArray<float> Eval(const TArray<float>& InputData);
+
 #if WITH_EDITORONLY_DATA
     TArray<uint32> NearestNeighborIds;
 #endif
@@ -61,4 +66,11 @@ private:
 
     UPROPERTY(Transient)
     TObjectPtr<UNearestNeighborOptimizedNetworkInstance> OptimizedNetworkInstance = nullptr;
+
+private:
+	void CreateNNEModel();
+    void RunNNE();
+	TUniquePtr<UE::NNECore::IModelCPU> ModelCPU;
+	TArray<float> NNEInputTensorBuffer;
+	TArray<float> NNEOutputTensorBuffer;
 };
diff --git a/Engine/Plugins/Animation/MLDeformer/NearestNeighborModel/Source/NearestNeighborModel/Public/NearestNeighborOptimizedNetworkLoader.h b/Engine/Plugins/Animation/MLDeformer/NearestNeighborModel/Source/NearestNeighborModel/Public/NearestNeighborOptimizedNetworkLoader.h
--- a/Engine/Plugins/Animation/MLDeformer/NearestNeighborModel/Source/NearestNeighborModel/Public/NearestNeighborOptimizedNetworkLoader.h
+++ b/Engine/Plugins/Animation/MLDeformer/NearestNeighborModel/Source/NearestNeighborModel/Public/NearestNeighborOptimizedNetworkLoader.h
@@ -20,9 +20,6 @@ public:
 	UFUNCTION(BlueprintImplementableEvent, Category = "Nearest Neighbor Model")
 	bool LoadOptimizedNetwork(const FString& OnnxPath);
 
-	UFUNCTION(BlueprintImplementableEvent, Category = "Nearest Neighbor Model")
-	bool DoesMeetPrerequisites() const;
-
 	void SetOptimizedNetwork(UNearestNeighborOptimizedNetwork* InNetwork);
 
 	UFUNCTION(BlueprintPure, Category = "Nearest Neighbor Model")
diff --git a/Engine/Plugins/Animation/MLDeformer/NearestNeighborModel/Source/NearestNeighborModelEditor/NearestNeighborModelEditor.Build.cs b/Engine/Plugins/Animation/MLDeformer/NearestNeighborModel/Source/NearestNeighborModelEditor/NearestNeighborModelEditor.Build.cs
--- a/Engine/Plugins/Animation/MLDeformer/NearestNeighborModel/Source/NearestNeighborModelEditor/NearestNeighborModelEditor.Build.cs
+++ b/Engine/Plugins/Animation/MLDeformer/NearestNeighborModel/Source/NearestNeighborModelEditor/NearestNeighborModelEditor.Build.cs
@@ -33,7 +33,7 @@ namespace UnrealBuildTool.Rules
 					"NearestNeighborModel",
 					"Projects",
 					"PropertyEditor",
-					"NeuralNetworkInference",
+					"NNECore",
 					"ToolWidgets",
 					"ComputeFramework",
 					"RenderCore",
diff --git a/Engine/Plugins/Animation/MLDeformer/NearestNeighborModel/Source/NearestNeighborModelEditor/Private/NearestNeighborEditorModel.cpp b/Engine/Plugins/Animation/MLDeformer/NearestNeighborModel/Source/NearestNeighborModelEditor/Private/NearestNeighborEditorModel.cpp
--- a/Engine/Plugins/Animation/MLDeformer/NearestNeighborModel/Source/NearestNeighborModelEditor/Private/NearestNeighborEditorModel.cpp
+++ b/Engine/Plugins/Animation/MLDeformer/NearestNeighborModel/Source/NearestNeighborModelEditor/Private/NearestNeighborEditorModel.cpp
@@ -13,7 +13,8 @@
 #include "MLDeformerComponent.h"
 #include "MLDeformerEditorToolkit.h"
 #include "MLDeformerAsset.h"
-#include "NeuralNetwork.h"
+#include "NNECore.h"
+#include "NNECoreModelData.h"
 #include "GeometryCache.h"
 #include "GeometryCacheComponent.h"
 #include "Animation/DebugSkelMeshComponent.h"
@@ -23,6 +24,7 @@
 #include "PackageTools.h"
 #include "ObjectTools.h"
 #include "UObject/SavePackage.h"
+#include "Misc/FileHelper.h"
 
 
 #define LOCTEXT_NAMESPACE "NearestNeighborEditorModel"
@@ -93,22 +95,24 @@ namespace UE::NearestNeighborModel
 		}
 	}
 
-	UNeuralNetwork* FNearestNeighborEditorModel::LoadNeuralNetworkFromOnnx(const FString& Filename) const
+	TObjectPtr<UNNEModelData> FNearestNeighborEditorModel::LoadNeuralNetworkFromOnnx(const FString& Filename) const
 	{
 		const FString OnnxFile = FPaths::ConvertRelativePathToFull(Filename);
 		if (FPaths::FileExists(OnnxFile))
 		{
 			UE_LOG(LogNearestNeighborModel, Display, TEXT("Loading Onnx file '%s'..."), *OnnxFile);
-			UNeuralNetwork* Result = NewObject<UNeuralNetwork>(Model, UNeuralNetwork::StaticClass());		
-			if (Result->Load(OnnxFile))
+			TArray<uint8> RawModelData;
+			const bool bIsModelInMem = FFileHelper::LoadFileToArray(RawModelData, *OnnxFile);
+
+			if (bIsModelInMem)
 			{
-				Result->SetDeviceType(ENeuralDeviceType::CPU, ENeuralDeviceType::CPU, ENeuralDeviceType::CPU);	
-				if (Result->GetDeviceType() != ENeuralDeviceType::CPU || Result->GetOutputDeviceType() != ENeuralDeviceType::CPU || Result->GetInputDeviceType() != ENeuralDeviceType::CPU)
-				{
-					UE_LOG(LogNearestNeighborModel, Error, TEXT("Neural net in ML Deformer '%s' cannot run on the CPU, it will not be active."), *Model->GetDeformerAsset()->GetName());
-				}
+				UNearestNeighborModel* NearestNeighborModel = GetNearestNeighborModel();
+				TObjectPtr<UNNEModelData> ModelData = NewObject<UNNEModelData>();
+
+				ModelData->Init(NearestNeighborModel->GetNNERuntimeName(), RawModelData);
+				// TODO - add some real verification
 				UE_LOG(LogNearestNeighborModel, Display, TEXT("Successfully loaded Onnx file '%s'..."), *OnnxFile);
-				return Result;
+				return ModelData;
 			}
 			else
 			{
@@ -122,43 +126,33 @@ namespace UE::NearestNeighborModel
 
 		return nullptr;
 	}
-
+	
 	bool FNearestNeighborEditorModel::LoadTrainedNetwork() const
 	{
 		UNearestNeighborModel* NearestNeighborModel = GetNearestNeighborModel();
 		if (NearestNeighborModel)
-		{
+		{	
 			const FString OnnxFile = GetTrainedNetworkOnnxFile();
-			UNeuralNetwork* Network = LoadNeuralNetworkFromOnnx(OnnxFile);
-			if (Network)
+			if (NearestNeighborModel->ShouldUseOptimizedNetwork())
 			{
-				if (NearestNeighborModel->DoesEditorSupportOptimizedNetwork())
+				const bool bSuccess = NearestNeighborModel->LoadOptimizedNetwork(OnnxFile);
+				if (bSuccess)
 				{
-					// We still need NNINetwork for serialization in nearest neighbor update and kmeans update.
-					NearestNeighborModel->SetNNINetwork(Network, false);
-
-					const bool bSuccess = NearestNeighborModel->LoadOptimizedNetwork(OnnxFile);
-					if (bSuccess)
+					UNearestNeighborModelInstance* ModelInstance = static_cast<UNearestNeighborModelInstance*>(GetTestMLDeformerModelInstance());
+					if (ModelInstance)
 					{
-						UNearestNeighborModelInstance* ModelInstance = static_cast<UNearestNeighborModelInstance*>(GetTestMLDeformerModelInstance());
-						if (ModelInstance)
-						{
-							ModelInstance->InitOptimizedNetworkInstance();
-							return true;
-						}
+						ModelInstance->InitOptimizedNetworkInstance();
+						NearestNeighborModel->SetUseOptimizedNetwork(true);
+						return true;
 					}
-					NearestNeighborModel->SetUseOptimizedNetwork(true);
-				}
-				else
-				{
-					NearestNeighborModel->SetNNINetwork(Network);
-					NearestNeighborModel->SetUseOptimizedNetwork(false);
-					return true;
 				}
 			}
 			else
 			{
-				return false;
+				TObjectPtr<UNNEModelData> NNEModel = LoadNeuralNetworkFromOnnx(OnnxFile);
+				NearestNeighborModel->SetUseOptimizedNetwork(false);
+				NearestNeighborModel->SetNNEModelData(NNEModel, false);
+				return true;
 			}
 		}
 		return false;
@@ -169,7 +163,7 @@ namespace UE::NearestNeighborModel
 		const UNearestNeighborModel* NearestNeighborModel = GetNearestNeighborModel();
 		if (NearestNeighborModel)
 		{
-			return NearestNeighborModel->DoesUseOptimizedNetwork() ? NearestNeighborModel->GetOptimizedNetwork() != nullptr : NearestNeighborModel->GetNNINetwork() != nullptr;
+			return NearestNeighborModel->DoesUseOptimizedNetwork() ? NearestNeighborModel->GetOptimizedNetwork() != nullptr : NearestNeighborModel->GetNNEModelData().Get() != nullptr;
 		}
 		return false;
 	}
@@ -454,19 +448,6 @@ namespace UE::NearestNeighborModel
 			return EUpdateResult::SUCCESS;
 		}
 
-		const UNeuralNetwork* NeuralNetwork = NearestNeighborModel->GetNNINetwork();
-		if(NeuralNetwork && NeuralNetwork->IsLoaded())
-		{
-			const FString SavePath = GetTrainedNetworkOnnxFile();
-			UE_LOG(LogNearestNeighborModel, Display, TEXT("Saving to %s"), *SavePath);
-			NeuralNetwork->Save(SavePath);
-		}
-		else
-		{
-			UE_LOG(LogNearestNeighborModel, Display, TEXT("Network not loaded. A network needs to be trained."));
-			return EUpdateResult::WARNING;
-		}
-
 		FNearestNeighborGeomCacheSampler* GeomCacheSampler = static_cast<FNearestNeighborGeomCacheSampler*>(GetGeomCacheSampler());
 		if (GeomCacheSampler)
 		{
@@ -548,28 +529,6 @@ namespace UE::NearestNeighborModel
 		UNearestNeighborModel* NearestNeighborModel = static_cast<UNearestNeighborModel*>(Model);
 		check(NearestNeighborModel != nullptr);
 
-		const UNeuralNetwork* NeuralNetwork = NearestNeighborModel->GetNNINetwork();
-		if(NeuralNetwork && NeuralNetwork->IsLoaded())
-		{
-			const FString SavePath = GetTrainedNetworkOnnxFile();
-			if (FPaths::DirectoryExists(FPaths::GetPath(SavePath)))
-			{
-				UE_LOG(LogNearestNeighborModel, Display, TEXT("Saving to %s"), *SavePath);
-				NeuralNetwork->Save(SavePath);
-			}
-			else
-			{
-				UE_LOG(LogNearestNeighborModel, Error, TEXT("Path %s does not exist."), *FPaths::GetPath(SavePath));
-				KMeansClusterResult |= EUpdateResult::ERROR;
-				return;
-			}
-		}
-		else
-		{
-			UE_LOG(LogNearestNeighborModel, Warning, TEXT("Network is not available. Nothing will be done."));
-			KMeansClusterResult |= EUpdateResult::WARNING;
-			return;
-		}
 		if (NearestNeighborModel->KMeansPartId >= NearestNeighborModel->GetNumParts())
 		{
 			UE_LOG(LogNearestNeighborModel, Error, TEXT("KMeansPartId %d is out of range [0, %d). Nothing will be done."), NearestNeighborModel->KMeansPartId, NearestNeighborModel->GetNumParts());
@@ -724,15 +683,6 @@ namespace UE::NearestNeighborModel
 		MorphTargetUpdateResult = EUpdateResult::SUCCESS;
 
 		UNearestNeighborModel* NearestNeighborModel = GetNearestNeighborModel();
-		// Automatic converting NNI network to optimized network
-		const UNeuralNetwork* NeuralNetwork = NearestNeighborModel->GetNNINetwork();
-		if (NearestNeighborModel->DoesEditorSupportOptimizedNetwork() && NeuralNetwork != nullptr && NearestNeighborModel->GetOptimizedNetwork() == nullptr)
-		{
-			const FString SavePath = GetTrainedNetworkOnnxFile();
-			UE_LOG(LogNearestNeighborModel, Display, TEXT("Saving to %s"), *SavePath);
-			NeuralNetwork->Save(SavePath);
-			LoadTrainedNetwork();
-		}
 
 		if (!NearestNeighborModel->IsClothPartDataValid())
 		{
@@ -755,7 +705,7 @@ namespace UE::NearestNeighborModel
 			}
 		}
 
-		if (IsNeuralNetworkLoaded())
+		if (IsTrained())
 		{
 			MorphTargetUpdateResult |= InitMorphTargets();
 			if (HasError(MorphTargetUpdateResult))
@@ -766,6 +716,7 @@ namespace UE::NearestNeighborModel
 			GetNearestNeighborModel()->UpdateNetworkSize();
 			GetNearestNeighborModel()->UpdateMorphTargetSize();
 		}
+
 		InitTestMLDeformerPreviousWeights();
 		GetNearestNeighborModel()->ValidateMorphTargetData();
 
@@ -925,31 +876,24 @@ namespace UE::NearestNeighborModel
 		}
 		else
 		{
-			const UNeuralNetwork* NeuralNetwork = NearestNeighborModel->GetNNINetwork();
-			const UMLDeformerComponent* MLDeformerComponent = GetTestMLDeformerComponent();
-			if(NeuralNetwork && NeuralNetwork->IsLoaded() && MLDeformerComponent && MLDeformerComponent->GetModelInstance())
+			const TObjectPtr<UNNEModelData> NNEModelData = NearestNeighborModel->GetNNEModelData();
+			if (NNEModelData == nullptr)
 			{
-				const int32 NeuralNetworkInferenceHandle = MLDeformerComponent->GetModelInstance()->GetNeuralNetworkInferenceHandle();
-				const FNeuralTensor& OutputTensor = NeuralNetwork->GetOutputTensorForContext(NeuralNetworkInferenceHandle);
-				const int32 NumNetworkWeights = OutputTensor.Num();
-				const int32 NumPCACoeffs = NearestNeighborModel->GetTotalNumPCACoeffs();
-				if (NumNetworkWeights != NumPCACoeffs)
-				{
-					UE_LOG(LogNearestNeighborModel, Warning, TEXT("Network output dimension %d is not equal to number of morph targets %d. Network needs to be re-trained and no deformation will be applied."), NumNetworkWeights, NumPCACoeffs);
-					return EUpdateResult::WARNING;
-				}
+				UE_LOG(LogNearestNeighborModel, Warning, TEXT("NNE network is not loaded. Model needs to be re-trained."));
+				return EUpdateResult::WARNING;
 			}
-		}
 
+			if (NNEModelData->GetModelData(NearestNeighborModel->GetNNERuntimeName()).IsEmpty())
+			{
+				UE_LOG(LogNearestNeighborModel, Warning, TEXT("NNE runtime is empty. Model needs to be re-trained."));
+				return EUpdateResult::WARNING;
+			}
+
+			// TODO: check network dimensions			
+		}
 		return EUpdateResult::SUCCESS;
 	}
 
-	bool FNearestNeighborEditorModel::IsNeuralNetworkLoaded()
-	{
-		const UNearestNeighborModel *NearestNeighborModel = static_cast<UNearestNeighborModel*>(Model);
-		const UNeuralNetwork* NeuralNetwork = NearestNeighborModel->GetNNINetwork();
-		return NeuralNetwork && NeuralNetwork->IsLoaded();
-	}
 
 	UMLDeformerModelInstance* FNearestNeighborEditorModel::GetTestMLDeformerModelInstance() const
 	{
diff --git a/Engine/Plugins/Animation/MLDeformer/NearestNeighborModel/Source/NearestNeighborModelEditor/Private/NearestNeighborTrainingModel.cpp b/Engine/Plugins/Animation/MLDeformer/NearestNeighborModel/Source/NearestNeighborModelEditor/Private/NearestNeighborTrainingModel.cpp
--- a/Engine/Plugins/Animation/MLDeformer/NearestNeighborModel/Source/NearestNeighborModelEditor/Private/NearestNeighborTrainingModel.cpp
+++ b/Engine/Plugins/Animation/MLDeformer/NearestNeighborModel/Source/NearestNeighborModelEditor/Private/NearestNeighborTrainingModel.cpp
@@ -2,6 +2,7 @@
 
 #include "NearestNeighborTrainingModel.h"
 #include "NearestNeighborModel.h"
+#include "NearestNeighborModelInstance.h"
 #include "NearestNeighborEditorModel.h"
 #include "NearestNeighborGeomCacheSampler.h"
 #include "Animation/AnimSequence.h"
@@ -111,4 +112,20 @@ const TArray<int32> UNearestNeighborTrainingModel::GetMeshIndexBuffer() const
 	return TArray<int32>((int32*)UIndexBuffer.GetData(), UIndexBuffer.Num());
 }
 
+UNearestNeighborModelInstance* UNearestNeighborTrainingModel::CreateModelInstance()
+{
+	UNearestNeighborModelInstance* ModelInstance = NewObject<UNearestNeighborModelInstance>(this);
+	ModelInstance->SetModel(GetNearestNeighborModel());
+	ModelInstance->Init(nullptr);
+	ModelInstance->PostMLDeformerComponentInit();
+	return ModelInstance;
+}
+
+void UNearestNeighborTrainingModel::DestroyModelInstance(UNearestNeighborModelInstance* ModelInstance)
+{
+	ModelInstance->Release();
+	ModelInstance->ConditionalBeginDestroy();
+}
+
+
 #undef LOCTEXT_NAMESPACE
diff --git a/Engine/Plugins/Animation/MLDeformer/NearestNeighborModel/Source/NearestNeighborModelEditor/Public/NearestNeighborEditorModel.h b/Engine/Plugins/Animation/MLDeformer/NearestNeighborModel/Source/NearestNeighborModelEditor/Public/NearestNeighborEditorModel.h
--- a/Engine/Plugins/Animation/MLDeformer/NearestNeighborModel/Source/NearestNeighborModelEditor/Public/NearestNeighborEditorModel.h
+++ b/Engine/Plugins/Animation/MLDeformer/NearestNeighborModel/Source/NearestNeighborModelEditor/Public/NearestNeighborEditorModel.h
@@ -15,6 +15,7 @@ class UNearestNeighborModel;
 class UNearestNeighborModelInstance;
 class UNearestNeighborTrainingModel;
 class UNearestNeighborModelVizSettings;
+class UNNEModelData;
 
 namespace UE::NearestNeighborModel
 {
@@ -61,6 +62,7 @@ namespace UE::NearestNeighborModel
 		UNearestNeighborModel* GetNearestNeighborModel() const { return static_cast<UNearestNeighborModel*>(Model); }
 		UNearestNeighborModelVizSettings* GetNearestNeighborModelVizSettings() const;
 
+		TObjectPtr<UNNEModelData> LoadNeuralNetworkFromOnnx(const FString& Filename) const;
 		// Recomputes nearest neighbor coeffcients and nearest neighbor vertex offsets. 
 		virtual uint8 UpdateNearestNeighborData();
 
@@ -81,14 +83,12 @@ namespace UE::NearestNeighborModel
 
 		int32 GetNumParts();
 
-		UNeuralNetwork* LoadNeuralNetworkFromOnnx(const FString& Filename) const;
 		void OnMorphTargetUpdate();
 		uint8 GetMorphTargetUpdateResult() { return MorphTargetUpdateResult; }
 		UMLDeformerComponent* GetTestMLDeformerComponent() const;
 		UMLDeformerModelInstance* GetTestMLDeformerModelInstance() const;
 		void InitTestMLDeformerPreviousWeights();
 		uint8 WarnIfNetworkInvalid();
-		bool IsNeuralNetworkLoaded();
 
 		virtual void CreateNearestNeighborActors(UWorld* World, int32 StartIndex = 0);
 
diff --git a/Engine/Plugins/Animation/MLDeformer/NearestNeighborModel/Source/NearestNeighborModelEditor/Public/NearestNeighborTrainingModel.h b/Engine/Plugins/Animation/MLDeformer/NearestNeighborModel/Source/NearestNeighborModelEditor/Public/NearestNeighborTrainingModel.h
--- a/Engine/Plugins/Animation/MLDeformer/NearestNeighborModel/Source/NearestNeighborModelEditor/Public/NearestNeighborTrainingModel.h
+++ b/Engine/Plugins/Animation/MLDeformer/NearestNeighborModel/Source/NearestNeighborModelEditor/Public/NearestNeighborTrainingModel.h
@@ -12,6 +12,7 @@ namespace UE::NearestNeighborModel
 }
 
 class UNearestNeighborModel;
+class UNearestNeighborModelInstance;
 
 UCLASS(Blueprintable)
 class NEARESTNEIGHBORMODELEDITOR_API UNearestNeighborTrainingModel
@@ -80,5 +81,11 @@ private:
 	UFUNCTION(BlueprintPure, Category = "Python")
 	const TArray<int32> GetMeshIndexBuffer() const;
 
+	UFUNCTION(BlueprintCallable, Category = "Python")
+	UNearestNeighborModelInstance* CreateModelInstance();
+
+	UFUNCTION(BlueprintCallable, Category = "Python")
+	void DestroyModelInstance(UNearestNeighborModelInstance* ModelInstance);
+
 	UNearestNeighborModel* NearestNeighborModel = nullptr;
 };
