NearestNeighborModel: remove NNE

diff --git a/Engine/Plugins/Animation/MLDeformer/NearestNeighborModel/NearestNeighborModel.uplugin b/Engine/Plugins/Animation/MLDeformer/NearestNeighborModel/NearestNeighborModel.uplugin
--- a/Engine/Plugins/Animation/MLDeformer/NearestNeighborModel/NearestNeighborModel.uplugin
+++ b/Engine/Plugins/Animation/MLDeformer/NearestNeighborModel/NearestNeighborModel.uplugin
@@ -14,21 +14,16 @@
 	"IsBetaVersion": false,
 	"IsExperimentalVersion": true,
 	"Installed": false,
-	"SupportedTargetPlatforms": [
-		"Win64", "Linux", "Mac"
-	],
 	"Modules": [
 		{
 			"Name": "NearestNeighborModel",
 			"Type": "Runtime",
 			"LoadingPhase": "Default",
-			"PlatformAllowList": [ "Win64", "Linux", "Mac" ]
 		},
 		{
 			"Name": "NearestNeighborModelEditor",
 			"Type": "Editor",
 			"LoadingPhase": "Default",
-			"PlatformAllowList": [ "Win64", "Linux", "Mac" ]
 		}
 	],
 	"Plugins": [
@@ -36,10 +31,6 @@
 			"Name": "GeometryCache",
 			"Enabled": true
 		},
-		{
-			"Name": "NNE",
-			"Enabled": true
-		},
 		{
 			"Name": "ComputeFramework",
 			"Enabled": true
diff --git a/Engine/Plugins/Animation/MLDeformer/NearestNeighborModel/Source/NearestNeighborModel/NearestNeighborModel.Build.cs b/Engine/Plugins/Animation/MLDeformer/NearestNeighborModel/Source/NearestNeighborModel/NearestNeighborModel.Build.cs
--- a/Engine/Plugins/Animation/MLDeformer/NearestNeighborModel/Source/NearestNeighborModel/NearestNeighborModel.Build.cs
+++ b/Engine/Plugins/Animation/MLDeformer/NearestNeighborModel/Source/NearestNeighborModel/NearestNeighborModel.Build.cs
@@ -38,9 +38,6 @@ public class NearestNeighborModel : ModuleRules
 				"CoreUObject",
 				"Engine",
 				"GeometryCache",
-				"NNECore",
-				"NNERuntimeRDG",
-				"NNEUtils",
 				"OptimusCore",
 				"Projects",
 				"RenderCore",
diff --git a/Engine/Plugins/Animation/MLDeformer/NearestNeighborModel/Source/NearestNeighborModel/Private/NearestNeighborModel.cpp b/Engine/Plugins/Animation/MLDeformer/NearestNeighborModel/Source/NearestNeighborModel/Private/NearestNeighborModel.cpp
--- a/Engine/Plugins/Animation/MLDeformer/NearestNeighborModel/Source/NearestNeighborModel/Private/NearestNeighborModel.cpp
+++ b/Engine/Plugins/Animation/MLDeformer/NearestNeighborModel/Source/NearestNeighborModel/Private/NearestNeighborModel.cpp
@@ -16,7 +16,6 @@
 #include "Rendering/MorphTargetVertexInfoBuffers.h"
 #include "ShaderCore.h"
 #include "UObject/UObjectGlobals.h"
-#include "NNECoreModelData.h"
 #include UE_INLINE_GENERATED_CPP_BY_NAME(NearestNeighborModel)
 
 #define LOCTEXT_NAMESPACE "UNearestNeighborModel"
@@ -65,16 +64,6 @@ UMLDeformerInputInfo* UNearestNeighborModel::CreateInputInfo()
 	return NearestNeighborModelInputInfo;
 }
 
-void UNearestNeighborModel::SetNNEModelData(TObjectPtr<UNNEModelData> ModelData)
-{
-	NNEModel = ModelData;
-}
-
-TObjectPtr<UNNEModelData> UNearestNeighborModel::GetNNEModelData() const
-{
-	return NNEModel;
-}
-
 UMLDeformerModelInstance* UNearestNeighborModel::CreateModelInstance(UMLDeformerComponent* Component)
 {
 	return NewObject<UNearestNeighborModelInstance>(Component);
@@ -83,23 +72,6 @@ UMLDeformerModelInstance* UNearestNeighborModel::CreateModelInstance(UMLDeformer
 void UNearestNeighborModel::Serialize(FArchive& Archive)
 {
 	Super::Serialize(Archive);
-	if (Archive.IsCooking())
-	{
-		if (DoesUseOptimizedNetwork())
-		{
-			if (GetNNEModelData() != nullptr)
-			{
-				NNEModel = nullptr;
-			}
-		}
-		else
-		{
-			if (GetOptimizedNetwork() != nullptr)
-			{
-				OptimizedNetwork = nullptr;
-			}
-		}
-	}
 }
 
 void UNearestNeighborModel::PostLoad()
@@ -413,16 +385,8 @@ int32 UNearestNeighborModel::GetNumNeighborsFromAnimSequence(int32 PartId) const
 
 void UNearestNeighborModel::UpdateNetworkSize()
 {
-	if (NNEModel)
-	{
-		const FString RuntimeName = GetNNERuntimeName();
-		TConstArrayView<uint8> ArrayView = NNEModel->GetModelData(RuntimeName);
-		SavedNetworkSize = ArrayView.GetTypeSize() * ArrayView.Num(); 
-	}
-	else
-	{
-		SavedNetworkSize = 0.0f;
-	}
+
+	SavedNetworkSize = 0.0f;
 }
 
 void UNearestNeighborModel::UpdateMorphTargetSize()
diff --git a/Engine/Plugins/Animation/MLDeformer/NearestNeighborModel/Source/NearestNeighborModel/Private/NearestNeighborModelInstance.cpp b/Engine/Plugins/Animation/MLDeformer/NearestNeighborModel/Source/NearestNeighborModel/Private/NearestNeighborModelInstance.cpp
--- a/Engine/Plugins/Animation/MLDeformer/NearestNeighborModel/Source/NearestNeighborModel/Private/NearestNeighborModelInstance.cpp
+++ b/Engine/Plugins/Animation/MLDeformer/NearestNeighborModel/Source/NearestNeighborModel/Private/NearestNeighborModelInstance.cpp
@@ -3,8 +3,6 @@
 #include "NearestNeighborModelInstance.h"
 #include "NearestNeighborModel.h"
 #include "NearestNeighborModelInputInfo.h"
-#include "NNECore.h"
-#include "NNECoreModelData.h"
 #include "NearestNeighborOptimizedNetwork.h"
 #include "MLDeformerAsset.h"
 #include "Components/ExternalMorphSet.h"
@@ -62,10 +60,7 @@ bool UNearestNeighborModelInstance::SetupInputs()
 	}
 	else
 	{
-		if (!ModelCPU || NNEInputTensorBuffer.IsEmpty() || NNEOutputTensorBuffer.IsEmpty())
-		{
-			return false;
-		}
+		return false;
 	}
 
 	int32 NumNeuralNetInputs = 0;
@@ -106,10 +101,6 @@ void UNearestNeighborModelInstance::Init(USkeletalMeshComponent* SkelMeshCompone
 	{
 		InitOptimizedNetworkInstance();
 	}
-	else
-	{
-		CreateNNEModel(); 
-	}
 }
 
 void UNearestNeighborModelInstance::InitOptimizedNetworkInstance()
@@ -130,10 +121,6 @@ void UNearestNeighborModelInstance::Execute(float ModelWeight)
 		UNearestNeighborOptimizedNetwork* OptimizedNetwork = NearestNeighborModel->GetOptimizedNetwork();
 		OptimizedNetworkInstance->Run();
 	}
-	else
-	{
-		RunNNE();
-	}
 }
 
 void UNearestNeighborModelInstance::Tick(float DeltaTime, float ModelWeight)
@@ -299,45 +286,6 @@ void UNearestNeighborModelInstance::InitPreviousWeights()
 	}
 }
 
-void UNearestNeighborModelInstance::CreateNNEModel()
-{
-	UNearestNeighborModel* NearestNeighborModel = Cast<UNearestNeighborModel>(Model);
-	if (NearestNeighborModel)
-	{
-		TWeakInterfacePtr<INNERuntime> Runtime = UE::NNECore::GetRuntime<INNERuntime>(NearestNeighborModel->GetNNERuntimeName());
-		TWeakInterfacePtr<INNERuntimeCPU> RuntimeCPU = UE::NNECore::GetRuntime<INNERuntimeCPU>(NearestNeighborModel->GetNNERuntimeName());
-
-		if (!Runtime.IsValid())
-		{
-			UE_LOG(LogNNE, Error, TEXT("Can't get %s runtime."), *NearestNeighborModel->GetNNERuntimeName());
-			return;
-		}
-
-		TObjectPtr<UNNEModelData> ModelData = NearestNeighborModel->GetNNEModelData(); 
-		if (ModelData)
-		{
-			if (RuntimeCPU.IsValid())
-			{
-				// allocate tensor inputs and outputs
-				ModelCPU = RuntimeCPU->CreateModelCPU(ModelData);
-
-				// setup inputs
-				TConstArrayView<UE::NNECore::FTensorDesc> InputTensorDescs = ModelCPU->GetInputTensorDescs();
-				UE::NNECore::FTensorShape InputTensorShape = UE::NNECore::FTensorShape::MakeFromSymbolic(InputTensorDescs[0].GetShape());
-				ModelCPU->SetInputTensorShapes({ InputTensorShape });
-				check(InputTensorDescs[0].GetElemByteSize() == sizeof(float)); 
-				NNEInputTensorBuffer.SetNumUninitialized(InputTensorShape.Volume());
-
-				// setup outputs
-				TConstArrayView<UE::NNECore::FTensorDesc> OutputTensorDescs = ModelCPU->GetOutputTensorDescs();
-				UE::NNECore::FTensorShape OutputTensorShape = UE::NNECore::FTensorShape::MakeFromSymbolic(OutputTensorDescs[0].GetShape());
-				check(OutputTensorDescs[0].GetElemByteSize() == sizeof(float));
-				NNEOutputTensorBuffer.SetNumUninitialized(OutputTensorShape.Volume());
-			}
-		}
-	}
-}
-
 void UNearestNeighborModelInstance::GetInputDataPointer(float*& OutInputData, int32& OutNumInputFloats)
 {
 	const UNearestNeighborModel* NearestNeighborModel = Cast<UNearestNeighborModel>(Model);
@@ -352,14 +300,6 @@ void UNearestNeighborModelInstance::GetInputDataPointer(float*& OutInputData, in
 			return;
 		}
 	}
-	else
-	{
-		if (ModelCPU)
-		{
-			OutInputData = NNEInputTensorBuffer.GetData();
-			OutNumInputFloats = NNEInputTensorBuffer.Num(); 
-		}
-	}
 	OutInputData = nullptr;
 	OutNumInputFloats = 0;
 }
@@ -378,15 +318,6 @@ void UNearestNeighborModelInstance::GetOutputDataPointer(float*& OutOutputData,
 			return;
 		}
 	}
-	else
-	{
-		if (ModelCPU)
-		{
-			OutOutputData = NNEOutputTensorBuffer.GetData();
-			OutNumOutputFloats = NNEOutputTensorBuffer.Num();
-			return;
-		}
-	}
 	OutOutputData = nullptr;
 	OutNumOutputFloats = 0;
 }
@@ -427,13 +358,6 @@ FString UNearestNeighborModelInstance::CheckCompatibility(USkeletalMeshComponent
 		}
 
 	}
-	else
-	{
-		if (ModelCPU)
-		{
-			NumNeuralNetInputs = NNEInputTensorBuffer.Num();
-		}
-	}
 
 	// Only check compatibility after the network is loaded. 
 	if (NumNeuralNetInputs >= 0 && NumNeuralNetInputs != NumDeformerAssetInputs) 
@@ -471,17 +395,5 @@ TArray<float> UNearestNeighborModelInstance::Eval(const TArray<float>& InputData
 		TArrayView<float> Outputs = OptimizedNetworkInstance->GetOutputs();
 		return TArray<float>(Outputs.GetData(), Outputs.Num());
 	}
-	else
-	{
-		check(ModelCPU != nullptr);
-		RunNNE();
-		return NNEOutputTensorBuffer;
-	}
-}
-
-void UNearestNeighborModelInstance::RunNNE()
-{
-	UE::NNECore::FTensorBindingCPU InputTensorBinding{NNEInputTensorBuffer.GetData(), NNEInputTensorBuffer.Num()};
-	UE::NNECore::FTensorBindingCPU OutputTensorBinding{NNEOutputTensorBuffer.GetData(), NNEOutputTensorBuffer.Num()};
-	ModelCPU->RunSync({InputTensorBinding}, {OutputTensorBinding});
+	return TArray<float>();
 }
\ No newline at end of file
diff --git a/Engine/Plugins/Animation/MLDeformer/NearestNeighborModel/Source/NearestNeighborModel/Public/NearestNeighborModel.h b/Engine/Plugins/Animation/MLDeformer/NearestNeighborModel/Source/NearestNeighborModel/Public/NearestNeighborModel.h
--- a/Engine/Plugins/Animation/MLDeformer/NearestNeighborModel/Source/NearestNeighborModel/Public/NearestNeighborModel.h
+++ b/Engine/Plugins/Animation/MLDeformer/NearestNeighborModel/Source/NearestNeighborModel/Public/NearestNeighborModel.h
@@ -14,7 +14,6 @@ struct FSkelMeshImportedMeshInfo;
 class USkeletalMesh;
 class UGeometryCache;
 class UAnimSequence;
-class UNNEModelData;
 class UMLDeformerAsset;
 class USkeleton;
 class IPropertyHandle;
@@ -261,11 +260,6 @@ private:
 	void SetOptimizedNetwork(UNearestNeighborOptimizedNetwork* InOptimizedNetwork);
 	bool LoadOptimizedNetwork(const FString& OnnxPath);
 
-	void SetNNEModelData(TObjectPtr<UNNEModelData> ModelData);
-	TObjectPtr<UNNEModelData> GetNNEModelData() const;
-
-	const FString GetNNERuntimeName() const {return TEXT("NNERuntimeORTDml"); }
-
 #if WITH_EDITORONLY_DATA
 	TObjectPtr<UAnimSequence> GetNearestNeighborSkeletons(int32 PartId);
 	const TObjectPtr<UAnimSequence> GetNearestNeighborSkeletons(int32 PartId) const;
@@ -455,10 +449,6 @@ private:
 	UPROPERTY()
 	TObjectPtr<UNearestNeighborOptimizedNetwork> OptimizedNetwork = nullptr;
 
-	/** The NNE neural network model. */
-	UPROPERTY()
-	TObjectPtr<UNNEModelData> NNEModel;
-
 	UPROPERTY()
 	bool bUseOptimizedNetwork = true;
 };
diff --git a/Engine/Plugins/Animation/MLDeformer/NearestNeighborModel/Source/NearestNeighborModel/Public/NearestNeighborModelInstance.h b/Engine/Plugins/Animation/MLDeformer/NearestNeighborModel/Source/NearestNeighborModel/Public/NearestNeighborModelInstance.h
--- a/Engine/Plugins/Animation/MLDeformer/NearestNeighborModel/Source/NearestNeighborModel/Public/NearestNeighborModelInstance.h
+++ b/Engine/Plugins/Animation/MLDeformer/NearestNeighborModel/Source/NearestNeighborModel/Public/NearestNeighborModelInstance.h
@@ -2,8 +2,6 @@
 #pragma once
 
 #include "MLDeformerMorphModelInstance.h"
-#include "NNECore.h"
-#include "NNECoreRuntimeCPU.h"
 #include "NearestNeighborModelInstance.generated.h"
 
 
@@ -66,11 +64,4 @@ private:
 
     UPROPERTY(Transient)
     TObjectPtr<UNearestNeighborOptimizedNetworkInstance> OptimizedNetworkInstance = nullptr;
-
-private:
-	void CreateNNEModel();
-    void RunNNE();
-	TUniquePtr<UE::NNECore::IModelCPU> ModelCPU;
-	TArray<float> NNEInputTensorBuffer;
-	TArray<float> NNEOutputTensorBuffer;
 };
diff --git a/Engine/Plugins/Animation/MLDeformer/NearestNeighborModel/Source/NearestNeighborModelEditor/NearestNeighborModelEditor.Build.cs b/Engine/Plugins/Animation/MLDeformer/NearestNeighborModel/Source/NearestNeighborModelEditor/NearestNeighborModelEditor.Build.cs
--- a/Engine/Plugins/Animation/MLDeformer/NearestNeighborModel/Source/NearestNeighborModelEditor/NearestNeighborModelEditor.Build.cs
+++ b/Engine/Plugins/Animation/MLDeformer/NearestNeighborModel/Source/NearestNeighborModelEditor/NearestNeighborModelEditor.Build.cs
@@ -33,7 +33,6 @@ namespace UnrealBuildTool.Rules
 					"NearestNeighborModel",
 					"Projects",
 					"PropertyEditor",
-					"NNECore",
 					"ToolWidgets",
 					"ComputeFramework",
 					"RenderCore",
diff --git a/Engine/Plugins/Animation/MLDeformer/NearestNeighborModel/Source/NearestNeighborModelEditor/Private/NearestNeighborEditorModel.cpp b/Engine/Plugins/Animation/MLDeformer/NearestNeighborModel/Source/NearestNeighborModelEditor/Private/NearestNeighborEditorModel.cpp
--- a/Engine/Plugins/Animation/MLDeformer/NearestNeighborModel/Source/NearestNeighborModelEditor/Private/NearestNeighborEditorModel.cpp
+++ b/Engine/Plugins/Animation/MLDeformer/NearestNeighborModel/Source/NearestNeighborModelEditor/Private/NearestNeighborEditorModel.cpp
@@ -13,8 +13,6 @@
 #include "MLDeformerComponent.h"
 #include "MLDeformerEditorToolkit.h"
 #include "MLDeformerAsset.h"
-#include "NNECore.h"
-#include "NNECoreModelData.h"
 #include "GeometryCache.h"
 #include "GeometryCacheComponent.h"
 #include "Animation/DebugSkelMeshComponent.h"
@@ -94,38 +92,6 @@ namespace UE::NearestNeighborModel
 			GetEditor()->GetModelDetailsView()->ForceRefresh();
 		}
 	}
-
-	TObjectPtr<UNNEModelData> FNearestNeighborEditorModel::LoadNeuralNetworkFromOnnx(const FString& Filename) const
-	{
-		const FString OnnxFile = FPaths::ConvertRelativePathToFull(Filename);
-		if (FPaths::FileExists(OnnxFile))
-		{
-			UE_LOG(LogNearestNeighborModel, Display, TEXT("Loading Onnx file '%s'..."), *OnnxFile);
-			TArray<uint8> RawModelData;
-			const bool bIsModelInMem = FFileHelper::LoadFileToArray(RawModelData, *OnnxFile);
-
-			if (bIsModelInMem)
-			{
-				UNearestNeighborModel* NearestNeighborModel = GetNearestNeighborModel();
-				TObjectPtr<UNNEModelData> ModelData = NewObject<UNNEModelData>();
-
-				ModelData->Init(NearestNeighborModel->GetNNERuntimeName(), RawModelData);
-				// TODO - add some real verification
-				UE_LOG(LogNearestNeighborModel, Display, TEXT("Successfully loaded Onnx file '%s'..."), *OnnxFile);
-				return ModelData;
-			}
-			else
-			{
-				UE_LOG(LogNearestNeighborModel, Error, TEXT("Failed to load Onnx file '%s'"), *OnnxFile);
-			}
-		}
-		else
-		{
-			UE_LOG(LogNearestNeighborModel, Error, TEXT("Onnx file '%s' does not exist!"), *OnnxFile);
-		}
-
-		return nullptr;
-	}
 	
 	bool FNearestNeighborEditorModel::LoadTrainedNetwork() const
 	{
@@ -147,13 +113,6 @@ namespace UE::NearestNeighborModel
 					}
 				}
 			}
-			else
-			{
-				TObjectPtr<UNNEModelData> NNEModel = LoadNeuralNetworkFromOnnx(OnnxFile);
-				NearestNeighborModel->SetUseOptimizedNetwork(false);
-				NearestNeighborModel->SetNNEModelData(NNEModel);
-				return true;
-			}
 		}
 		return false;
 	}
@@ -163,7 +122,7 @@ namespace UE::NearestNeighborModel
 		const UNearestNeighborModel* NearestNeighborModel = GetNearestNeighborModel();
 		if (NearestNeighborModel)
 		{
-			return NearestNeighborModel->DoesUseOptimizedNetwork() ? NearestNeighborModel->GetOptimizedNetwork() != nullptr : NearestNeighborModel->GetNNEModelData().Get() != nullptr;
+			return NearestNeighborModel->DoesUseOptimizedNetwork() && NearestNeighborModel->GetOptimizedNetwork() != nullptr;
 		}
 		return false;
 	}
@@ -874,23 +833,6 @@ namespace UE::NearestNeighborModel
 				return EUpdateResult::WARNING;
 			}
 		}
-		else
-		{
-			const TObjectPtr<UNNEModelData> NNEModelData = NearestNeighborModel->GetNNEModelData();
-			if (NNEModelData == nullptr)
-			{
-				UE_LOG(LogNearestNeighborModel, Warning, TEXT("NNE network is not loaded. Model needs to be re-trained."));
-				return EUpdateResult::WARNING;
-			}
-
-			if (NNEModelData->GetModelData(NearestNeighborModel->GetNNERuntimeName()).IsEmpty())
-			{
-				UE_LOG(LogNearestNeighborModel, Warning, TEXT("NNE runtime is empty. Model needs to be re-trained."));
-				return EUpdateResult::WARNING;
-			}
-
-			// TODO: check network dimensions			
-		}
 		return EUpdateResult::SUCCESS;
 	}
 
diff --git a/Engine/Plugins/Animation/MLDeformer/NearestNeighborModel/Source/NearestNeighborModelEditor/Public/NearestNeighborEditorModel.h b/Engine/Plugins/Animation/MLDeformer/NearestNeighborModel/Source/NearestNeighborModelEditor/Public/NearestNeighborEditorModel.h
--- a/Engine/Plugins/Animation/MLDeformer/NearestNeighborModel/Source/NearestNeighborModelEditor/Public/NearestNeighborEditorModel.h
+++ b/Engine/Plugins/Animation/MLDeformer/NearestNeighborModel/Source/NearestNeighborModelEditor/Public/NearestNeighborEditorModel.h
@@ -15,7 +15,6 @@ class UNearestNeighborModel;
 class UNearestNeighborModelInstance;
 class UNearestNeighborTrainingModel;
 class UNearestNeighborModelVizSettings;
-class UNNEModelData;
 
 namespace UE::NearestNeighborModel
 {
@@ -62,7 +61,6 @@ namespace UE::NearestNeighborModel
 		UNearestNeighborModel* GetNearestNeighborModel() const { return static_cast<UNearestNeighborModel*>(Model); }
 		UNearestNeighborModelVizSettings* GetNearestNeighborModelVizSettings() const;
 
-		TObjectPtr<UNNEModelData> LoadNeuralNetworkFromOnnx(const FString& Filename) const;
 		// Recomputes nearest neighbor coeffcients and nearest neighbor vertex offsets. 
 		virtual uint8 UpdateNearestNeighborData();
 
