NearestNeighborModel: Remove trailing commas in uplugin file

diff --git a/Engine/Plugins/Animation/MLDeformer/NearestNeighborModel/NearestNeighborModel.uplugin b/Engine/Plugins/Animation/MLDeformer/NearestNeighborModel/NearestNeighborModel.uplugin
--- a/Engine/Plugins/Animation/MLDeformer/NearestNeighborModel/NearestNeighborModel.uplugin
+++ b/Engine/Plugins/Animation/MLDeformer/NearestNeighborModel/NearestNeighborModel.uplugin
@@ -18,12 +18,12 @@
 		{
 			"Name": "NearestNeighborModel",
 			"Type": "Runtime",
-			"LoadingPhase": "Default",
+			"LoadingPhase": "Default"
 		},
 		{
 			"Name": "NearestNeighborModelEditor",
 			"Type": "Editor",
-			"LoadingPhase": "Default",
+			"LoadingPhase": "Default"
 		}
 	],
 	"Plugins": [
