NearestNeighborModel: fix the +1 bug for vertex map

diff --git a/Engine/Plugins/Animation/MLDeformer/NearestNeighborModel/Source/NearestNeighborModel/Private/NearestNeighborModel.cpp b/Engine/Plugins/Animation/MLDeformer/NearestNeighborModel/Source/NearestNeighborModel/Private/NearestNeighborModel.cpp
--- a/Engine/Plugins/Animation/MLDeformer/NearestNeighborModel/Source/NearestNeighborModel/Private/NearestNeighborModel.cpp
+++ b/Engine/Plugins/Animation/MLDeformer/NearestNeighborModel/Source/NearestNeighborModel/Private/NearestNeighborModel.cpp
@@ -103,7 +103,7 @@ TArray<uint32> ReadTxt(const FString &Path)
 
 TArray<uint32> Range(const uint32 Start, const uint32 End)
 {
-	TArray<uint32> Arr; Arr.SetNum(End - Start + 1);
+	TArray<uint32> Arr; Arr.SetNum(End - Start);
 	for (uint32 i = Start; i < End; i++)
 	{
 		Arr[i - Start] = i;
