NearestNeighborModel: remove UsePartOnlyMesh

diff --git a/Engine/Plugins/Animation/MLDeformer/NearestNeighborModel/Source/NearestNeighborModel/Public/NearestNeighborModel.h b/Engine/Plugins/Animation/MLDeformer/NearestNeighborModel/Source/NearestNeighborModel/Public/NearestNeighborModel.h
--- a/Engine/Plugins/Animation/MLDeformer/NearestNeighborModel/Source/NearestNeighborModel/Public/NearestNeighborModel.h
+++ b/Engine/Plugins/Animation/MLDeformer/NearestNeighborModel/Source/NearestNeighborModel/Public/NearestNeighborModel.h
@@ -333,9 +333,6 @@ private:
 	static FName GetRecomputeDeltasPropertyName() { return GET_MEMBER_NAME_CHECKED(UNearestNeighborModel, bRecomputeDeltas); }
 	static FName GetRecomputePCAPropertyName() { return GET_MEMBER_NAME_CHECKED(UNearestNeighborModel, bRecomputePCA); }
 
-	bool GetUsePartOnlyMesh() const { return bUsePartOnlyMesh; }
-	static FName GetUsePartOnlyMeshPropertyName() { return GET_MEMBER_NAME_CHECKED(UNearestNeighborModel, bUsePartOnlyMesh); }
-
 	UFUNCTION(BlueprintPure, Category = "Nearest Neighbor Model")
 	FString GetModelDir() const;
 
@@ -384,9 +381,6 @@ protected:
 	UPROPERTY()
 	bool bMorphTargetDataValid = false;
 
-	UPROPERTY(EditAnywhere, BlueprintReadWrite, Category = "Nearest Neighbors")
-	bool bUsePartOnlyMesh = true;
-
 	UPROPERTY(EditAnywhere, BlueprintReadWrite, Category = "Nearest Neighbors")
 	TArray<FSkeletonCachePair> NearestNeighborData;
 
diff --git a/Engine/Plugins/Animation/MLDeformer/NearestNeighborModel/Source/NearestNeighborModelEditor/Private/NearestNeighborEditorModel.cpp b/Engine/Plugins/Animation/MLDeformer/NearestNeighborModel/Source/NearestNeighborModelEditor/Private/NearestNeighborEditorModel.cpp
--- a/Engine/Plugins/Animation/MLDeformer/NearestNeighborModel/Source/NearestNeighborModelEditor/Private/NearestNeighborEditorModel.cpp
+++ b/Engine/Plugins/Animation/MLDeformer/NearestNeighborModel/Source/NearestNeighborModelEditor/Private/NearestNeighborEditorModel.cpp
@@ -69,7 +69,7 @@ namespace UE::NearestNeighborModel
 			GetEditor()->GetModelDetailsView()->ForceRefresh();
 		}
 
-		if (Property->GetFName() == UNearestNeighborModel::GetNearestNeighborDataPropertyName() || (PropertyChangedEvent.MemberProperty != nullptr && PropertyChangedEvent.MemberProperty->GetFName() == UNearestNeighborModel::GetNearestNeighborDataPropertyName()) || Property->GetFName() == UNearestNeighborModel::GetUsePartOnlyMeshPropertyName())
+		if (Property->GetFName() == UNearestNeighborModel::GetNearestNeighborDataPropertyName() || (PropertyChangedEvent.MemberProperty != nullptr && PropertyChangedEvent.MemberProperty->GetFName() == UNearestNeighborModel::GetNearestNeighborDataPropertyName()))
 		{
 			GetNearestNeighborModel()->InvalidateNearestNeighborData();
 			GetEditor()->GetModelDetailsView()->ForceRefresh();
@@ -293,8 +293,7 @@ namespace UE::NearestNeighborModel
 					GeometryCacheComponent->SetManualTick(true);
 					GeometryCacheComponent->SetPlaybackSpeed(1.0f);
 					GeometryCacheComponent->Play();
-					uint8 ReturnCode = GeomCacheSampler->GeneratePartMeshMappings(GetNearestNeighborModel()->PartVertexMap(PartId), NearestNeighborModel->GetUsePartOnlyMesh());
-					if (HasError(ReturnCode))
+					if (uint8 ReturnCode = GeomCacheSampler->GenerateMeshMappings(); HasError(ReturnCode))
 					{
 						return ReturnCode;
 					}
diff --git a/Engine/Plugins/Animation/MLDeformer/NearestNeighborModel/Source/NearestNeighborModelEditor/Private/NearestNeighborGeomCacheSampler.cpp b/Engine/Plugins/Animation/MLDeformer/NearestNeighborModel/Source/NearestNeighborModelEditor/Private/NearestNeighborGeomCacheSampler.cpp
--- a/Engine/Plugins/Animation/MLDeformer/NearestNeighborModel/Source/NearestNeighborModelEditor/Private/NearestNeighborGeomCacheSampler.cpp
+++ b/Engine/Plugins/Animation/MLDeformer/NearestNeighborModel/Source/NearestNeighborModelEditor/Private/NearestNeighborGeomCacheSampler.cpp
@@ -171,253 +171,32 @@ namespace UE::NearestNeighborModel
 		}
 	}
 
-	uint8 FNearestNeighborGeomCacheSampler::SamplePart(int32 InAnimFrameIndex, int32 PartId)
-	{
-		FMLDeformerSampler::Sample(InAnimFrameIndex);
-		USkeletalMesh* SkeletalMesh = SkeletalMeshComponent.Get() ? SkeletalMeshComponent->GetSkeletalMeshAsset() : nullptr;
-		UGeometryCache* GeometryCache = GeometryCacheComponent.Get() ? GeometryCacheComponent->GetGeometryCache() : nullptr;
-		if (SkeletalMeshComponent && SkeletalMesh && GeometryCacheComponent && GeometryCache)
-		{
-			const FTransform& AlignmentTransform = Model->GetAlignmentTransform();
-			FSkeletalMeshModel* ImportedModel = SkeletalMesh->GetImportedModel();
-			const int32 LODIndex = 0;
-			const FSkeletalMeshLODModel& LODModel = ImportedModel->LODModels[LODIndex];
-			const TArray<FSkelMeshImportedMeshInfo>& SkelMeshInfos = LODModel.ImportedMeshInfos;
-
-			if (HasError(CheckMeshMappingsEmpty()))
-			{
-				return EUpdateResult::ERROR;
-			}
-
-			if (PartId >= MeshMappingIndices.Num())
-			{
-				UE_LOG(LogNearestNeighborModel, Error, TEXT("SamplePart: MeshMappingIndices.Num()=%d is smaller than PartId %d"), MeshMappingIndices.Num(), PartId);
-				return EUpdateResult::ERROR;
-			}
-
-			const UE::MLDeformer::FMLDeformerGeomCacheMeshMapping& MeshMapping = MeshMappings[MeshMappingIndices[PartId]]; 
-
-			check(SkelMeshInfos.Num() > MeshMapping.MeshIndex);
-			const FSkelMeshImportedMeshInfo& MeshInfo = SkelMeshInfos[MeshMapping.MeshIndex]; 
-			check(MeshInfo.StartImportedVertex == 0);
-
-			check(GeometryCache->Tracks.Num() > MeshMapping.TrackIndex);
-			UGeometryCacheTrack* Track = GeometryCache->Tracks[MeshMapping.TrackIndex];
-			GeomCacheMeshDatas.Reset(1);
-			GeomCacheMeshDatas.AddDefaulted(1);
-			FGeometryCacheMeshData& GeomCacheMeshData = GeomCacheMeshDatas[0];
-
-			if (!Track->GetMeshDataAtTime(SampleTime, GeomCacheMeshData))
-			{
-				UE_LOG(LogNearestNeighborModel, Error, TEXT("SamplePart: Track cannot get mesh delta at frame %d"), InAnimFrameIndex);
-				return EUpdateResult::ERROR;
-			}
-
-			// Calculate the vertex deltas.
-			const FSkeletalMeshLODRenderData& SkelMeshLODData = SkeletalMesh->GetResourceForRendering()->LODRenderData[LODIndex];
-			const FSkinWeightVertexBuffer& SkinWeightBuffer = *SkeletalMeshComponent->GetSkinWeightBuffer(LODIndex);
-
-			UNearestNeighborModel* NearestNeighborModel = static_cast<UNearestNeighborModel*>(Model);
-			check(NearestNeighborModel != nullptr);
-
-			if (NearestNeighborModel->DoesUseDualQuaternionDeltas())
-			{
-				UpdateSkinnedPositions();
-			}
-
-			const TArray<uint32>& VertexMap = NearestNeighborModel->PartVertexMap(PartId);
-			const int32 NumPartVerts = VertexMap.Num();
-			PartVertexDeltas.Reset();
-			PartVertexDeltas.SetNum(NumPartVerts * 3);
-
-			for(int32 PartVertexIndex = 0; PartVertexIndex < NumPartVerts; PartVertexIndex++)
-			{
-				const int32 VertexIndex = VertexMap[PartVertexIndex];
-				const int32 GeomCacheVertexIndex = MeshMapping.SkelMeshToTrackVertexMap[PartVertexIndex];
-
-				if (GeomCacheVertexIndex != INDEX_NONE && GeomCacheMeshData.Positions.IsValidIndex(GeomCacheVertexIndex))
-				{
-					FVector3f Delta = FVector3f::ZeroVector;
-
-					const int32 ArrayIndex = 3 * PartVertexIndex;
-					// Calculate the inverse skinning transform for this vertex.
-					const int32 RenderVertexIndex = MeshMapping.ImportedVertexToRenderVertexMap[PartVertexIndex];
-					if (RenderVertexIndex != INDEX_NONE)
-					{
-						const FVector3f GeomCacheVertexPos = (FVector3f)AlignmentTransform.TransformPosition((FVector)GeomCacheMeshData.Positions[GeomCacheVertexIndex]);
-						if (NearestNeighborModel->DoesUseDualQuaternionDeltas())
-						{
-							const FVector3f SkinnedVertexPos = SkinnedVertexPositions[VertexIndex];
-							const FVector3f WorldDelta = GeomCacheVertexPos - SkinnedVertexPos;
-							Delta = CalcDualQuaternionDelta(RenderVertexIndex, WorldDelta, SkelMeshLODData, SkinWeightBuffer);
-						}
-						else
-						{
-							const FMatrix44f InvSkinningTransform = CalcInverseSkinningTransform(RenderVertexIndex, SkelMeshLODData, SkinWeightBuffer);
-
-							// Calculate the pre-skinning data.
-							const FVector3f UnskinnedPosition = SkelMeshLODData.StaticVertexBuffers.PositionVertexBuffer.VertexPosition(RenderVertexIndex);
-							const FVector3f PreSkinningTargetPos = InvSkinningTransform.TransformPosition(GeomCacheVertexPos);
-							Delta = PreSkinningTargetPos - UnskinnedPosition;
-						}
-					}
-
-					PartVertexDeltas[ArrayIndex] = Delta.X;
-					PartVertexDeltas[ArrayIndex + 1] = Delta.Y;
-					PartVertexDeltas[ArrayIndex + 2] = Delta.Z;
-				}
-			}
-		}
-		return EUpdateResult::SUCCESS;
-	}
-
-	bool IsPotentialMatch(const FString& TrackName, const FString& MeshName)
-	{
-		return (TrackName.Find(MeshName) == 0);
-	}
-
-	uint8 FNearestNeighborGeomCacheSampler::GeneratePartMeshMappings(const TArray<uint32>& VertexMap, bool bUsePartOnlyMesh)
+	uint8 FNearestNeighborGeomCacheSampler::GenerateMeshMappings()
 	{
 		uint8 Result = EUpdateResult::SUCCESS;
 		USkeletalMesh* SkeletalMesh = SkeletalMeshComponent.Get() ? SkeletalMeshComponent->GetSkeletalMeshAsset() : nullptr;
 		UGeometryCache* GeometryCache = GeometryCacheComponent.Get() ? GeometryCacheComponent->GetGeometryCache() : nullptr;
 		if (SkeletalMeshComponent && SkeletalMesh && GeometryCacheComponent && GeometryCache)
 		{
-			if (!bUsePartOnlyMesh)
+			TArray<FString> FailedNames;
+			TArray<FString> VertexMisMatchNames;
+			GenerateGeomCacheMeshMappings(SkeletalMesh, GeometryCache, MeshMappings, FailedNames, VertexMisMatchNames);
+			if (!FailedNames.IsEmpty() || !VertexMisMatchNames.IsEmpty())
 			{
-				TArray<FString> FailedNames;
-				TArray<FString> VertexMisMatchNames;
-				GenerateGeomCacheMeshMappings(SkeletalMesh, GeometryCache, MeshMappings, FailedNames, VertexMisMatchNames);
-				Result |= GenerateMeshMappingIndices();
-				if (!FailedNames.IsEmpty() || !VertexMisMatchNames.IsEmpty())
-				{
-					Result |= EUpdateResult::WARNING;
-				}
-				for(int32 i = 0; i < VertexMisMatchNames.Num(); i++)
-				{
-					UE_LOG(LogNearestNeighborModel, Warning, TEXT("%s is skipped because it has different vertex counts in skeletal mesh and geometry cache."), *VertexMisMatchNames[i]);
-				}
-				GeomCacheMeshDatas.Reset();
-				GeomCacheMeshDatas.AddDefaulted(MeshMappings.Num());
-
-				Result |= CheckMeshMappingsEmpty();
-				return Result;
+				Result |= EUpdateResult::WARNING;
 			}
-			FSkeletalMeshModel* ImportedModel = SkeletalMesh->GetImportedModel();
-
-			if (!ImportedModel || ImportedModel->LODModels[0].ImportedMeshInfos.IsEmpty())
+			for(int32 i = 0; i < VertexMisMatchNames.Num(); i++)
 			{
-				UE_LOG(LogNearestNeighborModel, Error, TEXT("Unable to generate mesh mappings because SkeletalMesh has no imported model."));
-				return EUpdateResult::ERROR;
+				UE_LOG(LogNearestNeighborModel, Warning, TEXT("%s is skipped because it has different vertex counts in skeletal mesh and geometry cache."), *VertexMisMatchNames[i]);
 			}
-
-			const TArray<FSkelMeshImportedMeshInfo>& SkelMeshInfos = ImportedModel->LODModels[0].ImportedMeshInfos;
-			MeshMappings.Reset();
-
-			FString SkelMeshName;
-			const bool bIsSoloMesh = (GeometryCache->Tracks.Num() == 1 && SkelMeshInfos.Num() == 1);	// Do we just have one mesh and one track?
-
-			for (int32 TrackIndex = 0; TrackIndex < GeometryCache->Tracks.Num(); ++TrackIndex)
-			{
-				// Check if this is a candidate based on the mesh and track name.
-				UGeometryCacheTrack* Track = GeometryCache->Tracks[TrackIndex];
-
-				bool bFoundMatch = false;
-				for (int32 SkelMeshIndex = 0; SkelMeshIndex < SkelMeshInfos.Num(); ++SkelMeshIndex)
-				{
-					const FSkelMeshImportedMeshInfo& MeshInfo = SkelMeshInfos[SkelMeshIndex];
-
-					SkelMeshName = MeshInfo.Name.ToString();
-					if (Track &&
-						(IsPotentialMatch(Track->GetName(), SkelMeshName) || bIsSoloMesh))
-					{
-						// Extract the geom cache mesh data.
-						FGeometryCacheMeshData GeomCacheMeshData;
-						if (!Track->GetMeshDataAtTime(SampleTime, GeomCacheMeshData))
-						{
-							continue;
-						}
-
-						// Verify that we have imported vertex numbers.
-						if (GeomCacheMeshData.ImportedVertexNumbers.IsEmpty())
-						{
-							UE_LOG(LogNearestNeighborModel, Warning, TEXT("Geometry track %s is skipped because it has no imported vertex numbers."), *Track->GetName());
-							continue;
-						}
-
-						const int32 NumVertsFromGeomCache = FMath::Max(GeomCacheMeshData.ImportedVertexNumbers) + 1;
-						const int32 NumVertsFromVertexMap = VertexMap.Num();
-						Result |= CheckGeomCacheVertCount(NumVertsFromGeomCache, NumVertsFromVertexMap);
-						if (HasError(Result))
-						{
-							return Result;
-						}
-
-						// Create a new mesh mapping entry.
-						MeshMappings.AddDefaulted();
-						UE::MLDeformer::FMLDeformerGeomCacheMeshMapping& Mapping = MeshMappings.Last();
-						Mapping.MeshIndex = SkelMeshIndex;
-						Mapping.TrackIndex = TrackIndex;
-
-						const int32 NumPartVerts = VertexMap.Num();
-						Mapping.SkelMeshToTrackVertexMap.AddUninitialized(NumPartVerts);
-						Mapping.ImportedVertexToRenderVertexMap.AddUninitialized(NumPartVerts);
-
-
-						for(int32 PartVertexIndex = 0; PartVertexIndex < NumPartVerts; PartVertexIndex++)
-						{
-							Mapping.SkelMeshToTrackVertexMap[PartVertexIndex] = GeomCacheMeshData.ImportedVertexNumbers.Find(PartVertexIndex);
-							const int32 VertexIndex = VertexMap[PartVertexIndex];
-							const int32 RenderVertexIndex = ImportedModel->LODModels[0].MeshToImportVertexMap.Find(VertexIndex);
-							Mapping.ImportedVertexToRenderVertexMap[PartVertexIndex] = RenderVertexIndex;
-						}
-
-						// We found a match, no need to iterate over more MeshInfos.
-						bFoundMatch = true;
-						break;
-					} // If the track name matches the skeletal meshes internal mesh name.
-				} // For all meshes in the Skeletal Mesh.
-
-				if (Track && !bFoundMatch)
-				{
-					Result |= EUpdateResult::WARNING;
-					UE_LOG(LogNearestNeighborModel, Warning, TEXT("Geometry cache '%s' cannot be matched with a mesh inside the Skeletal Mesh."), *Track->GetName());
-				}
-			} // For all tracks.
+			GeomCacheMeshDatas.Reset();
+			GeomCacheMeshDatas.AddDefaulted(MeshMappings.Num());
+			
+			Result |= CheckMeshMappingsEmpty();
 		}
-		else
-		{
-			Result |= EUpdateResult::WARNING;
-			UE_LOG(LogNearestNeighborModel, Warning, TEXT("SkeletalMesh or GeometryCache is none. No mapping is generated"));
-		}
-		Result |= CheckMeshMappingsEmpty();
-		Result |= GenerateMeshMappingIndices();
 		return Result;
 	}
 
-	uint8 FNearestNeighborGeomCacheSampler::CheckGeomCacheVertCount(int32 NumVertsFromGeomCache, int32 NumVertsFromVertexMap) const
-	{
-		uint8 ResultCode = EUpdateResult::SUCCESS;		
-		UNearestNeighborModel* NearestNeighborModel = static_cast<UNearestNeighborModel*>(Model);
-		if (NearestNeighborModel->GetUsePartOnlyMesh())
-		{
-			if (NumVertsFromGeomCache != NumVertsFromVertexMap)
-			{
-				UE_LOG(LogNearestNeighborModel, Error, TEXT("Vertex count mismatch: GeomCache has %d vertices but vertex map has %d vertices. Maybe turn off UsePartOnlyMesh."), NumVertsFromGeomCache, NumVertsFromVertexMap);
-				return EUpdateResult::ERROR;
-			}
-		}
-		else
-		{
-			if (NumVertsFromVertexMap > NumVertsFromGeomCache)
-			{
-				UE_LOG(LogNearestNeighborModel, Error, TEXT("Vertex count in vertex map %d is larger than the vertex count in geometry cache %d. Something is wrong with vertex map or geometry cache"), NumVertsFromVertexMap, NumVertsFromGeomCache);
-				return EUpdateResult::ERROR;
-			}
-		}
-		return ResultCode;
-	}
-
 	uint8 FNearestNeighborGeomCacheSampler::CheckMeshMappingsEmpty() const
 	{
 		if(MeshMappings.IsEmpty())
@@ -499,35 +278,6 @@ namespace UE::NearestNeighborModel
 		return false;
 	}
 
-	uint8 FNearestNeighborGeomCacheSampler::GenerateMeshMappingIndices()
-	{
-		UNearestNeighborModel* NearestNeighborModel = static_cast<UNearestNeighborModel*>(Model);
-		if (NearestNeighborModel)
-		{
-			const int32 NumParts = NearestNeighborModel->GetNumParts();
-			MeshMappingIndices.SetNum(NumParts);
-			for (int32 PartId = 0; PartId < NumParts; PartId++)
-			{
-				bool bFoundIndex = false;
-				for (int32 MappingId = 0; MappingId < MeshMappings.Num(); MappingId++)
-				{
-					if (MeshMappings[MappingId].MeshIndex == NearestNeighborModel->GetPartMeshIndex(PartId))
-					{
-						MeshMappingIndices[PartId] = MappingId;
-						bFoundIndex = true;
-						break;
-					}
-				}
-				if (!bFoundIndex)
-				{
-					UE_LOG(LogNearestNeighborModel, Error, TEXT("Part %d could not find a mesh mapping."), PartId);
-					return EUpdateResult::ERROR;
-				}
-			}
-		}
-		return EUpdateResult::SUCCESS;
-	}
-
 	TArray<uint32> FNearestNeighborGeomCacheSampler::GetMeshIndexBuffer() const
 	{
 		TArray<uint32> IndexBuffer;
diff --git a/Engine/Plugins/Animation/MLDeformer/NearestNeighborModel/Source/NearestNeighborModelEditor/Private/NearestNeighborModelDetails.cpp b/Engine/Plugins/Animation/MLDeformer/NearestNeighborModel/Source/NearestNeighborModelEditor/Private/NearestNeighborModelDetails.cpp
--- a/Engine/Plugins/Animation/MLDeformer/NearestNeighborModel/Source/NearestNeighborModelEditor/Private/NearestNeighborModelDetails.cpp
+++ b/Engine/Plugins/Animation/MLDeformer/NearestNeighborModel/Source/NearestNeighborModelEditor/Private/NearestNeighborModelDetails.cpp
@@ -217,7 +217,6 @@ namespace UE::NearestNeighborModel
 		// Nearest Neighbor settings
 		NearestNeighborCategoryBuilder->AddProperty(GET_MEMBER_NAME_STRING_CHECKED(UNearestNeighborModel, DecayFactor));
 		NearestNeighborCategoryBuilder->AddProperty(GET_MEMBER_NAME_STRING_CHECKED(UNearestNeighborModel, NearestNeighborOffsetWeight));
-		NearestNeighborCategoryBuilder->AddProperty(UNearestNeighborModel::GetUsePartOnlyMeshPropertyName());
 		NearestNeighborCategoryBuilder->AddProperty(UNearestNeighborModel::GetNearestNeighborDataPropertyName());
 
 		KMeansCategoryBuilder->AddProperty(GET_MEMBER_NAME_STRING_CHECKED(UNearestNeighborModel, SourceAnims));
diff --git a/Engine/Plugins/Animation/MLDeformer/NearestNeighborModel/Source/NearestNeighborModelEditor/Private/NearestNeighborTrainingModel.cpp b/Engine/Plugins/Animation/MLDeformer/NearestNeighborModel/Source/NearestNeighborModelEditor/Private/NearestNeighborTrainingModel.cpp
--- a/Engine/Plugins/Animation/MLDeformer/NearestNeighborModel/Source/NearestNeighborModelEditor/Private/NearestNeighborTrainingModel.cpp
+++ b/Engine/Plugins/Animation/MLDeformer/NearestNeighborModel/Source/NearestNeighborModelEditor/Private/NearestNeighborTrainingModel.cpp
@@ -34,15 +34,6 @@ const TArray<int32> UNearestNeighborTrainingModel::GetPartVertexMap(const int32
 	return TArray<int32>((int32*)VertexMap.GetData(), VertexMap.Num());
 }
 
-int32 UNearestNeighborTrainingModel::SamplePart(int32 PartId, int32 Index)
-{
-	FNearestNeighborGeomCacheSampler* Sampler = static_cast<FNearestNeighborGeomCacheSampler*>(EditorModel->GetSampler());
-	int32 Result = Sampler->SamplePart(Index, PartId);
-	PartSampleDeltas = Sampler->GetPartVertexDeltas();
-	SampleBoneRotations = Sampler->GetBoneRotations();
-	return Result;
-}
-
 int32 UNearestNeighborTrainingModel::SetSamplerPartData(const int32 PartId)
 {
 	return GetNearestNeighborEditorModel()->SetSamplerPartData(PartId);
diff --git a/Engine/Plugins/Animation/MLDeformer/NearestNeighborModel/Source/NearestNeighborModelEditor/Public/NearestNeighborGeomCacheSampler.h b/Engine/Plugins/Animation/MLDeformer/NearestNeighborModel/Source/NearestNeighborModelEditor/Public/NearestNeighborGeomCacheSampler.h
--- a/Engine/Plugins/Animation/MLDeformer/NearestNeighborModel/Source/NearestNeighborModelEditor/Public/NearestNeighborGeomCacheSampler.h
+++ b/Engine/Plugins/Animation/MLDeformer/NearestNeighborModel/Source/NearestNeighborModelEditor/Public/NearestNeighborGeomCacheSampler.h
@@ -22,19 +22,13 @@ namespace UE::NearestNeighborModel
 
 	private:
 		virtual void SampleDualQuaternionDeltas(int32 InAnimFrameIndex);
-		virtual uint8 SamplePart(int32 InAnimFrameIndex, int32 PartId);
 		virtual bool SampleKMeansAnim(const int32 SkeletonId);
 		virtual bool SampleKMeansFrame(const int32 Frame);
-		const TArray<float>& GetPartVertexDeltas() const { return PartVertexDeltas; }
-		uint8 GeneratePartMeshMappings(const TArray<uint32>& VertexMap, bool bUsePartOnlyMesh);
-		uint8 GenerateMeshMappingIndices();
-		uint8 CheckGeomCacheVertCount(int32 NumVertsFromGeomCache, int32 NumVertsFromVertexMap) const;
+		uint8 GenerateMeshMappings();
 		uint8 CheckMeshMappingsEmpty() const;
 		TArray<uint32> GetMeshIndexBuffer() const;
 		FVector3f CalcDualQuaternionDelta(int32 VertexIndex, const FVector3f& WorldDelta, const FSkeletalMeshLODRenderData& SkelMeshLODData, const FSkinWeightVertexBuffer& SkinWeightBuffer) const;
 
-		TArray<float> PartVertexDeltas;
-		TArray<int32> MeshMappingIndices;
 		int32 KMeansAnimId;
 	};
 }
diff --git a/Engine/Plugins/Animation/MLDeformer/NearestNeighborModel/Source/NearestNeighborModelEditor/Public/NearestNeighborTrainingModel.h b/Engine/Plugins/Animation/MLDeformer/NearestNeighborModel/Source/NearestNeighborModelEditor/Public/NearestNeighborTrainingModel.h
--- a/Engine/Plugins/Animation/MLDeformer/NearestNeighborModel/Source/NearestNeighborModelEditor/Public/NearestNeighborTrainingModel.h
+++ b/Engine/Plugins/Animation/MLDeformer/NearestNeighborModel/Source/NearestNeighborModelEditor/Public/NearestNeighborTrainingModel.h
@@ -51,9 +51,6 @@ private:
 	UFUNCTION(BlueprintPure, Category = "Python")
 	const TArray<int32> GetPartVertexMap(const int32 PartId) const;
 
-	UFUNCTION(BlueprintCallable, Category = "Python")
-	int32 SamplePart(int32 PartId, int32 Index);
-
 	UFUNCTION(BlueprintCallable, Category = "Python")
 	int32 SetSamplerPartData(const int32 PartId);
 
