NearestNeighborModel: include additive tracks for kmeans

diff --git a/Engine/Plugins/Animation/MLDeformer/NearestNeighborModel/Source/NearestNeighborModelEditor/Private/NearestNeighborEditorModel.cpp b/Engine/Plugins/Animation/MLDeformer/NearestNeighborModel/Source/NearestNeighborModelEditor/Private/NearestNeighborEditorModel.cpp
--- a/Engine/Plugins/Animation/MLDeformer/NearestNeighborModel/Source/NearestNeighborModelEditor/Private/NearestNeighborEditorModel.cpp
+++ b/Engine/Plugins/Animation/MLDeformer/NearestNeighborModel/Source/NearestNeighborModelEditor/Private/NearestNeighborEditorModel.cpp
@@ -19,6 +19,7 @@
 #include "Animation/MorphTarget.h"
 #include "Animation/AnimSequence.h"
 #include "Animation/AnimInstance.h"
+#include "BonePose.h"
 #include "PackageTools.h"
 #include "ObjectTools.h"
 #include "UObject/SavePackage.h"
@@ -538,26 +539,98 @@ namespace UE::NearestNeighborModel
 		KmeansResults.Reset();
 	}
 
+	template<typename T>
+	TArray<T> Range(T End)
+	{
+		TArray<T> Result;
+		Result.SetNum(End);
+		for (uint32 i = 0; i < End; i++)
+		{
+			Result[i] = i;
+		}
+		return Result;
+	}
+
+	class FAnimEvaluator
+	{
+	public:
+		FAnimEvaluator(USkeleton* Skeleton)
+		{
+			const FReferenceSkeleton& ReferenceSkeleton = Skeleton->GetReferenceSkeleton();
+			const int32 NumBones = ReferenceSkeleton.GetNum();
+			TArray<uint16> BoneIndices = Range<uint16>(NumBones);
+			BoneContainer.SetUseRAWData(true);
+			BoneContainer.InitializeTo(BoneIndices, UE::Anim::FCurveFilterSettings(), *Skeleton);
+			OutPose.SetBoneContainer(&BoneContainer);
+			OutCurve.InitFrom(BoneContainer);
+		}
+
+		TArray<FTransform> GetBoneTransforms(const UAnimSequence* Anim, int32 Frame)
+		{
+			const double Time = FMath::Clamp(Anim->GetSamplingFrameRate().AsSeconds(Frame), 0., (double)Anim->GetPlayLength());
+			FAnimExtractContext ExtractionContext(Time);
+			FAnimationPoseData AnimationPoseData(OutPose, OutCurve, TempAttributes);
+			Anim->GetAnimationPose(AnimationPoseData, ExtractionContext);
+
+			const int32 NumBones = BoneContainer.GetNumBones();
+			TArray<FTransform> Transforms;
+			Transforms.SetNum(NumBones);
+			for (int32 BoneIndex = 0; BoneIndex < NumBones; ++BoneIndex)
+			{
+				const FCompactPoseBoneIndex CompactIndex = BoneContainer.MakeCompactPoseIndex(FMeshPoseBoneIndex(BoneIndex));
+				const FTransform BoneTransform = OutPose[CompactIndex];
+				Transforms[BoneIndex] = BoneTransform;
+			}
+			return Transforms;
+		}
+
+	private:
+		FBoneContainer BoneContainer;
+		FCompactPose OutPose;
+		FBlendedCurve OutCurve;
+		UE::Anim::FStackAttributeContainer TempAttributes;
+	};
+
 	TPair<UAnimSequence*, uint8> FNearestNeighborEditorModel::CreateAnimOfClusterCenters(const FString& PackageName, const TArray<int32>& KmeansResults)
 	{
 		uint8 ReturnCode = EUpdateResult::SUCCESS;
 		UNearestNeighborModel *NearestNeighborModel = static_cast<UNearestNeighborModel*>(Model);
+		TTuple<UAnimSequence*, uint8> None = TTuple<UAnimSequence*, uint8>(nullptr, EUpdateResult::ERROR);
 
 		if (KmeansResults.Num() != NearestNeighborModel->NumClusters * 2)
 		{
 			UE_LOG(LogNearestNeighborModel, Error, TEXT("KmeansClusterPoses returned %d clusters whereas %d are expected."), KmeansResults.Num() / 2, NearestNeighborModel->NumClusters);
-			return TTuple<UAnimSequence*, uint8>(nullptr, EUpdateResult::ERROR);
+			return None;
+		}
+		if (NearestNeighborModel->SourceAnims.Num() == 0)
+		{
+			UE_LOG(LogNearestNeighborModel, Error, TEXT("No source anims found."));
+			return None;
 		}
-
 		const UAnimSequence* DefaultAnim = NearestNeighborModel->SourceAnims[0];
+		if (DefaultAnim == nullptr)
+		{
+			UE_LOG(LogNearestNeighborModel, Error, TEXT("Source anim 0 is null."));
+			return None;
+		}
 
 		UAnimSequence* Anim = CreateObjectInstance<UAnimSequence>(PackageName);
 		if (Anim == nullptr)
 		{
-			return TTuple<UAnimSequence*, uint8>(nullptr, EUpdateResult::ERROR);
+			UE_LOG(LogNearestNeighborModel, Error, TEXT("Failed to create AnimSequence."));
+			return None;
+		}
+
+		USkeleton* const Skeleton = DefaultAnim->GetSkeleton();
+		if (Skeleton == nullptr)
+		{
+			UE_LOG(LogNearestNeighborModel, Error, TEXT("Skeleton is null."));
+			return None;
 		}
+		const FReferenceSkeleton& ReferenceSkeleton = Skeleton->GetReferenceSkeleton();
+		const int32 NumBones = ReferenceSkeleton.GetNum();
 
-		Anim->SetSkeleton(DefaultAnim->GetSkeleton());
+		Anim->SetSkeleton(Skeleton);
 		IAnimationDataController& Controller = Anim->GetController();
 		Controller.OpenBracket(LOCTEXT("CreateNewAnim_Bracket", "Create New Anim"));
 		Controller.InitializeModel();
@@ -568,42 +641,55 @@ namespace UE::NearestNeighborModel
 		Controller.SetNumberOfFrames(NumKeys - 1);
 		Controller.SetFrameRate(FFrameRate(30, 1));
 
-		TArray<FName> TrackNames;		
-		DefaultAnim->GetDataModel()->GetBoneTrackNames(TrackNames);
-		const int32 NumTracks = TrackNames.Num();
-		for (const FName& TrackName : TrackNames)
-		{
-			TArray<FVector3f> PosKeys;
-			TArray<FQuat4f> RotKeys;
-			TArray<FVector3f> ScaleKeys;
-			PosKeys.SetNum(NumKeys);
-			RotKeys.SetNum(NumKeys);
-			ScaleKeys.SetNum(NumKeys);
-			for (int32 KeyIndex = 0; KeyIndex < NumKeys; KeyIndex++)
-			{
-				const int32 PickedAnimId = KmeansResults[KeyIndex * 2];
-				const int32 PickedFrame = KmeansResults[KeyIndex * 2 + 1];
-				if (PickedAnimId < 0 || PickedAnimId >= NearestNeighborModel->SourceAnims.Num())
-				{
-					UE_LOG(LogNearestNeighborModel, Error, TEXT("CreateAnimOfClusterCenters: PickedAnimId %d is out of range."), PickedAnimId);
-					return TTuple<UAnimSequence*, uint8>(nullptr, EUpdateResult::ERROR);
-				}
+		FMemMark Mark(FMemStack::Get());
+		FAnimEvaluator AnimEval(Skeleton);
 
-				const UAnimSequence* PickedAnim = NearestNeighborModel->SourceAnims[PickedAnimId];
-				if (PickedAnim == nullptr)
-				{
-					UE_LOG(LogNearestNeighborModel, Error, TEXT("CreateAnimOfClusterCenters: PickedAnim %d is null."), PickedAnimId);
-					return TTuple<UAnimSequence*, uint8>(nullptr, EUpdateResult::ERROR);
-				}
+		TArray<TArray<FVector3f>> PosKeys;
+		TArray<TArray<FQuat4f>> RotKeys;
+		TArray<TArray<FVector3f>> ScaleKeys;
+		PosKeys.SetNum(NumBones);
+		RotKeys.SetNum(NumBones);
+		ScaleKeys.SetNum(NumBones);
+		for (int32 BoneIndex = 0; BoneIndex < NumBones; ++BoneIndex)
+		{
+			PosKeys[BoneIndex].SetNum(NumKeys);
+			RotKeys[BoneIndex].SetNum(NumKeys);
+			ScaleKeys[BoneIndex].SetNum(NumKeys);
+		}
 
-				const FTransform BoneTransform = PickedAnim->GetDataModel()->GetBoneTrackTransform(TrackName, PickedFrame);				
-				PosKeys[KeyIndex] = FVector3f(BoneTransform.GetLocation());
-				RotKeys[KeyIndex] = FQuat4f(BoneTransform.GetRotation());
-				ScaleKeys[KeyIndex] = FVector3f(BoneTransform.GetScale3D());
+		for (int32 KeyIndex = 0; KeyIndex < NumKeys; ++KeyIndex)
+		{
+			const int32 PickedAnimId = KmeansResults[KeyIndex * 2];
+			const int32 PickedFrame = KmeansResults[KeyIndex * 2 + 1];
+			if (PickedAnimId < 0 || PickedAnimId >= NearestNeighborModel->SourceAnims.Num())
+			{
+				UE_LOG(LogNearestNeighborModel, Error, TEXT("CreateAnimOfClusterCenters: PickedAnimId %d is out of range."), PickedAnimId);
+				return None;
+			}
+			
+			const UAnimSequence* PickedAnim = NearestNeighborModel->SourceAnims[PickedAnimId];
+			if (PickedAnim == nullptr)
+			{
+				UE_LOG(LogNearestNeighborModel, Error, TEXT("CreateAnimOfClusterCenters: PickedAnim %d is null."), PickedAnimId);
+				return None;
+			}
+
+			const TArray<FTransform> BoneTransforms = AnimEval.GetBoneTransforms(PickedAnim, PickedFrame);
+			for (int32 BoneIndex = 0; BoneIndex < NumBones; ++BoneIndex)
+			{
+				const FTransform& BoneTransform = BoneTransforms[BoneIndex];
+				PosKeys[BoneIndex][KeyIndex] = FVector3f(BoneTransform.GetLocation());
+				RotKeys[BoneIndex][KeyIndex] = FQuat4f(BoneTransform.GetRotation());
+				ScaleKeys[BoneIndex][KeyIndex] = FVector3f(BoneTransform.GetScale3D());
 			}
-			Controller.AddBoneCurve(TrackName);
-			Controller.SetBoneTrackKeys(TrackName, PosKeys, RotKeys, ScaleKeys);
 		}
+		for (int32 BoneIndex = 0; BoneIndex < NumBones; ++BoneIndex)
+		{
+			const FName BoneName = ReferenceSkeleton.GetBoneName(BoneIndex);
+			Controller.AddBoneCurve(BoneName);
+			Controller.SetBoneTrackKeys(BoneName, PosKeys[BoneIndex], RotKeys[BoneIndex], ScaleKeys[BoneIndex]);
+		}
+
 		Controller.NotifyPopulated();
 		Controller.CloseBracket();
 		return MakeTuple(Anim, ReturnCode);
