NearestNeighborModel: Update statistics

diff --git a/Engine/Plugins/Animation/MLDeformer/NearestNeighborModel/Source/NearestNeighborModelEditor/Private/NearestNeighborEditorModel.cpp b/Engine/Plugins/Animation/MLDeformer/NearestNeighborModel/Source/NearestNeighborModelEditor/Private/NearestNeighborEditorModel.cpp
--- a/Engine/Plugins/Animation/MLDeformer/NearestNeighborModel/Source/NearestNeighborModelEditor/Private/NearestNeighborEditorModel.cpp
+++ b/Engine/Plugins/Animation/MLDeformer/NearestNeighborModel/Source/NearestNeighborModelEditor/Private/NearestNeighborEditorModel.cpp
@@ -863,6 +863,8 @@ namespace UE::NearestNeighborModel
 			MorphTarget->ConditionalBeginDestroy();
 		}
 
+		NearestNeighborModel->UpdateStatistics();
+
 		NearestNeighborModel->SetMorphTargetDeltas(Deltas);
 		return Result;
 	}
