MLDeformer: Updated NearestNeighborModel to use new NNERuntimeBasicCpu plugin

diff --git a/Engine/Plugins/Animation/MLDeformer/NearestNeighborModel/NearestNeighborModel.uplugin b/Engine/Plugins/Animation/MLDeformer/NearestNeighborModel/NearestNeighborModel.uplugin
--- a/Engine/Plugins/Animation/MLDeformer/NearestNeighborModel/NearestNeighborModel.uplugin
+++ b/Engine/Plugins/Animation/MLDeformer/NearestNeighborModel/NearestNeighborModel.uplugin
@@ -46,6 +46,10 @@
 		{
 			"Name": "MLDeformerFramework",
 			"Enabled": true
+		},
+		{
+			"Name": "NNERuntimeBasicCpu",
+			"Enabled": true
 		}
 	]
 }
\ No newline at end of file
diff --git a/Engine/Plugins/Animation/MLDeformer/NearestNeighborModel/Source/NearestNeighborModel/NearestNeighborModel.Build.cs b/Engine/Plugins/Animation/MLDeformer/NearestNeighborModel/Source/NearestNeighborModel/NearestNeighborModel.Build.cs
--- a/Engine/Plugins/Animation/MLDeformer/NearestNeighborModel/Source/NearestNeighborModel/NearestNeighborModel.Build.cs
+++ b/Engine/Plugins/Animation/MLDeformer/NearestNeighborModel/Source/NearestNeighborModel/NearestNeighborModel.Build.cs
@@ -43,6 +43,8 @@ public class NearestNeighborModel : ModuleRules
 				"RenderCore",
 				"RHI",
 				"MLDeformerFramework",
+				"NNE",
+				"NNERuntimeBasicCpu",
 			}
 			);
 		
diff --git a/Engine/Plugins/Animation/MLDeformer/NearestNeighborModel/Source/NearestNeighborModel/Private/NearestNeighborModel.cpp b/Engine/Plugins/Animation/MLDeformer/NearestNeighborModel/Source/NearestNeighborModel/Private/NearestNeighborModel.cpp
--- a/Engine/Plugins/Animation/MLDeformer/NearestNeighborModel/Source/NearestNeighborModel/Private/NearestNeighborModel.cpp
+++ b/Engine/Plugins/Animation/MLDeformer/NearestNeighborModel/Source/NearestNeighborModel/Private/NearestNeighborModel.cpp
@@ -437,11 +437,7 @@ FString UNearestNeighborModel::GetModelDir() const
 
 bool UNearestNeighborModel::ShouldUseOptimizedNetwork() const
 {
-#if NEARESTNEIGHBORMODEL_USE_ISPC
 	return UE::NearestNeighborModel::bNearestNeighborModelUseOptimizedNetwork;
-#else
-	return false;
-#endif
 }
 
 void UNearestNeighborModel::SetUseOptimizedNetwork(bool bInUseOptimizedNetwork)
@@ -535,6 +531,37 @@ bool UNearestNeighborModel::LoadOptimizedNetwork(const FString& OnnxPath)
 	return false;
 }
 
+bool UNearestNeighborModel::LoadOptimizedNetworkFromFile(const FString& Filename)
+{
+	const FString FullPath = FPaths::ConvertRelativePathToFull(Filename);
+	if (FPaths::FileExists(FullPath))
+	{
+		UE_LOG(LogNearestNeighborModel, Display, TEXT("Loading Network file '%s'..."), *FullPath);
+
+		// When we create a new UNearestNeighborOptimizedNetwork object we need to set the input
+		// and output sizes as it cannot get them from the model on load.
+		UNearestNeighborOptimizedNetwork* Result = NewObject<UNearestNeighborOptimizedNetwork>(this);
+		Result->SetNumInputs(InputDim);
+		Result->SetNumOutputs(OutputDim);
+		const bool bSuccess = Result->Load(FullPath);
+		
+		if (bSuccess)
+		{
+			// Clear optimized network to avoid error messages in SetOptimizedNetwork. 
+			OptimizedNetwork = nullptr;
+			SetOptimizedNetwork(Result);
+			return true;
+		}
+	}
+	else
+	{
+		UE_LOG(LogNearestNeighborModel, Error, TEXT("Network file '%s' does not exist!"), *FullPath);
+	}
+
+	return false;
+}
+
+
 int32 UNearestNeighborModel::GetOptimizedNetworkNumOutputs() const
 {
 	return OptimizedNetwork ? OptimizedNetwork->GetNumOutputs() : 0;
diff --git a/Engine/Plugins/Animation/MLDeformer/NearestNeighborModel/Source/NearestNeighborModel/Private/NearestNeighborOptimizedNetwork.cpp b/Engine/Plugins/Animation/MLDeformer/NearestNeighborModel/Source/NearestNeighborModel/Private/NearestNeighborOptimizedNetwork.cpp
--- a/Engine/Plugins/Animation/MLDeformer/NearestNeighborModel/Source/NearestNeighborModel/Private/NearestNeighborOptimizedNetwork.cpp
+++ b/Engine/Plugins/Animation/MLDeformer/NearestNeighborModel/Source/NearestNeighborModel/Private/NearestNeighborOptimizedNetwork.cpp
@@ -1,110 +1,212 @@
 // Copyright Epic Games, Inc. All Rights Reserved.
 
 #include "NearestNeighborOptimizedNetwork.h"
-#if NEARESTNEIGHBORMODEL_USE_ISPC
-#include "NearestNeighborOptimizedNetwork.ispc.generated.h"
-#endif
+
+#include "Misc/FileHelper.h"
+
+#include "NNE.h"
+#include "NNERuntime.h"
+#include "NNERuntimeCPU.h"
+#include "NNEModelData.h"
+#include "NNERuntimeBasicCpu.h"
+
 //--------------------------------------------------------------------------
-// UNearestNeighborNetworkLayer
+// UNearestNeighborOptimizedNetwork
 //--------------------------------------------------------------------------
-void UNearestNeighborNetworkLayer::AddParameter(const TArray<float>& Values, const TArray<int32>& Shape)
-{
-	Parameters.Add({Values, Shape});
-}
 
-void UNearestNeighborNetworkLayer::Run(const float* RESTRICT InputBuffer, float* RESTRICT OutputBuffer) const
+namespace UE::NearestNeighborModel::Private
 {
-}
+	PRAGMA_DISABLE_DEPRECATION_WARNINGS
 
-void UNearestNeighborNetworkLayer_Gemm_Prelu::Run(const float* RESTRICT InputBuffer, float* RESTRICT OutputBuffer) const
-{
-	const float* Gemm_Weights = Parameters[0].Values.GetData();
-	const float* Gemm_Bias = Parameters[1].Values.GetData();
-	const float PRelu_Slope = Parameters[2].Values[0];
-#if NEARESTNEIGHBORMODEL_USE_ISPC
-	ispc::Gemm_PRelu(OutputBuffer, InputBuffer, Gemm_Weights, Gemm_Bias, PRelu_Slope, NumInputs, NumOutputs);
-#endif
+	/** Creates the FileData from the legacy network format and clears it. */
+	static inline void CreateFileDataAndClearLayers(TArray<uint8>& OutFileData, TArray<TObjectPtr<UNearestNeighborNetworkLayer>>& Layers)
+	{
+		UE::NNE::RuntimeBasic::FSequentialModelBuilder Builder;
+
+		// Allocate data for the PReLuAlpha values as these need to be kept around until we run "Builder.WriteAndReset"
+		TArray<TArray<float>> PreluAlphaData;
+
+		for (UNearestNeighborNetworkLayer* Layer : Layers)
+		{
+			if (UNearestNeighborNetworkLayer_Gemm_Prelu* GemmPreluLayer = Cast<UNearestNeighborNetworkLayer_Gemm_Prelu>(Layer))
+			{
+				Builder.AddLinear(
+					GemmPreluLayer->NumInputs,
+					GemmPreluLayer->NumOutputs,
+					GemmPreluLayer->Parameters[0].Values,
+					GemmPreluLayer->Parameters[1].Values);
+
+				PreluAlphaData.AddDefaulted();
+				PreluAlphaData.Last().Init(GemmPreluLayer->Parameters[2].Values[0], GemmPreluLayer->NumOutputs);
+
+				Builder.AddPReLU(PreluAlphaData.Last());
+			}
+			else if (UNearestNeighborNetworkLayer_Gemm* GemmLayer = Cast<UNearestNeighborNetworkLayer_Gemm>(Layer))
+			{
+				Builder.AddLinear(
+					GemmLayer->NumInputs,
+					GemmLayer->NumOutputs,
+					GemmLayer->Parameters[0].Values,
+					GemmLayer->Parameters[1].Values);
+			}
+			else
+			{
+				checkf(false, TEXT("Unknown Layer Type"));
+			}
+		}
+
+		Layers.Empty();
+
+		OutFileData.SetNumUninitialized(Builder.GetWriteByteNum());
+		Builder.WriteAndReset(OutFileData);
+	}
+
+	PRAGMA_ENABLE_DEPRECATION_WARNINGS
 }
 
-void UNearestNeighborNetworkLayer_Gemm::Run(const float* RESTRICT InputBuffer, float* RESTRICT OutputBuffer) const
+PRAGMA_DISABLE_DEPRECATION_WARNINGS
+
+const TCHAR* UNearestNeighborOptimizedNetwork::RuntimeName = TEXT("NNERuntimeBasicCpu");
+
+
+void UNearestNeighborOptimizedNetwork::PostLoad()
 {
-	const float* Gemm_Weights = Parameters[0].Values.GetData();
-	const float* Gemm_Bias = Parameters[1].Values.GetData();
-#if NEARESTNEIGHBORMODEL_USE_ISPC
-	ispc::Gemm(OutputBuffer, InputBuffer, Gemm_Weights, Gemm_Bias, NumInputs, NumOutputs);
+	Super::PostLoad();
+
+	if (Layers_DEPRECATED.Num() > 0)
+	{
+		// NumInputs and NumOutputs are not used in the legacy format so we need to load them here.
+		NumInputs = Layers_DEPRECATED[0]->NumInputs;
+		NumOutputs = Layers_DEPRECATED.Last()->NumOutputs;
+
+		TArray<uint8> FileData;
+		UE::NearestNeighborModel::Private::CreateFileDataAndClearLayers(FileData, Layers_DEPRECATED);
+
+		if (!ModelData)
+		{
+			ModelData = NewObject<UNNEModelData>(this);
+		}
+
+		ModelData->Init(TEXT("ubnne"), FileData);
+
+		// The ModelData object stores a copy of the FileData so we manually clear this copy to avoid 
+		// having multiple copies of the network in memory at once during loading.
+		FileData.Empty();
+	}
+
+	// Create in-memory representation of network
+
+	TWeakInterfacePtr<INNERuntimeCPU> RuntimeCPU = UE::NNE::GetRuntime<INNERuntimeCPU>(RuntimeName);
+
+	if (ensureMsgf(RuntimeCPU.IsValid(), TEXT("Could not find requested NNE Runtime")))
+	{
+		if (ModelData)
+		{
+			Model = RuntimeCPU->CreateModelCPU(ModelData);
+		}
+	}
+
+	// If we are not in the editor then we clear the FileData and FileType since these will be
+	// using additional memory if we are loading from the legacy format.
+
+#if !WITH_EDITOR
+	ModelData->ClearFileDataAndFileType();
 #endif
 }
 
-//--------------------------------------------------------------------------
-// UNearestNeighborOptimizedNetwork
-//--------------------------------------------------------------------------
+PRAGMA_ENABLE_DEPRECATION_WARNINGS
+
 void UNearestNeighborOptimizedNetwork::Empty()
 {
-	Layers.Empty();
+	if (ModelData)
+	{
+		ModelData->ConditionalBeginDestroy();
+		ModelData = nullptr;
+		Model.Reset();
+	}
 }
 
 bool UNearestNeighborOptimizedNetwork::IsEmpty() const
 {
-	return Layers.IsEmpty();
+	return Model == nullptr;
 }
 
-int32 UNearestNeighborOptimizedNetwork::GetNumInputs() const
+bool UNearestNeighborOptimizedNetwork::Load(const FString& Filename)
 {
-	return !Layers.IsEmpty() ? Layers[0]->NumInputs : 0;
+	Empty();
+
+	TArray<uint8> FileData;
+	if (FFileHelper::LoadFileToArray(FileData, *Filename))
+	{
+		if (!ModelData)
+		{
+			ModelData = NewObject<UNNEModelData>(this);
+		}
+
+		ModelData->Init(TEXT("ubnne"), FileData);
+
+		// Clear FileData to avoid multiple copies in memory at once
+		FileData.Empty();
+
+		TWeakInterfacePtr<INNERuntimeCPU> RuntimeCPU = UE::NNE::GetRuntime<INNERuntimeCPU>(RuntimeName);
+
+		if (ensureMsgf(RuntimeCPU.IsValid(), TEXT("Could not find requested NNE Runtime")))
+		{
+			if (ModelData)
+			{
+				Model = RuntimeCPU->CreateModelCPU(ModelData);
+			}
+		}
+	}
+	else
+	{
+		return false;
+	}
+
+	return true;
 }
 
-int32 UNearestNeighborOptimizedNetwork::GetNumOutputs() const
+int32 UNearestNeighborOptimizedNetwork::GetNumInputs() const
 {
-	return !Layers.IsEmpty() ? Layers[Layers.Num()-1]->NumOutputs : 0;
+	return NumInputs;
 }
 
-bool UNearestNeighborOptimizedNetwork::Load(const FString& Filename)
+int32 UNearestNeighborOptimizedNetwork::GetNumOutputs() const
 {
-	return true;
+	return NumOutputs;
 }
 
-UNearestNeighborOptimizedNetworkInstance* UNearestNeighborOptimizedNetwork::CreateInstance()
+void UNearestNeighborOptimizedNetwork::SetNumInputs(int32 InNumInputs)
 {
-	UNearestNeighborOptimizedNetworkInstance* Instance = NewObject<UNearestNeighborOptimizedNetworkInstance>(this);
-	Instance->Init(this);
-	return Instance;
+	NumInputs = InNumInputs;
 }
 
-const int32 UNearestNeighborOptimizedNetwork::GetNumLayers() const
+void UNearestNeighborOptimizedNetwork::SetNumOutputs(int32 InNumOutputs)
 {
-	return Layers.Num();
+	NumOutputs = InNumOutputs;
 }
 
-UNearestNeighborNetworkLayer* UNearestNeighborOptimizedNetwork::GetLayer(int32 Index) const
+UNearestNeighborOptimizedNetworkInstance* UNearestNeighborOptimizedNetwork::CreateInstance()
 {
-	return Layers[Index].Get();
+	UNearestNeighborOptimizedNetworkInstance* Instance = NewObject<UNearestNeighborOptimizedNetworkInstance>(this);
+	Instance->Init(this);
+	return Instance;
 }
 
-UNearestNeighborNetworkLayer* UNearestNeighborOptimizedNetwork::AddLayer(const int32 LayerType)
+UE::NNE::IModelCPU* UNearestNeighborOptimizedNetwork::GetModel() const
 {
-	UNearestNeighborNetworkLayer* Layer = nullptr;
-	switch ((ENearestNeighborNetworkLayerType)LayerType)
-	{
-		case ENearestNeighborNetworkLayerType::Gemm_Prelu:
-			Layer = NewObject<UNearestNeighborNetworkLayer_Gemm_Prelu>(this);
-			break;
-		case ENearestNeighborNetworkLayerType::Gemm:
-			Layer = NewObject<UNearestNeighborNetworkLayer_Gemm>(this);
-			break;
-		default:
-			break;
-	}
-	if (Layer)
-	{
-		Layers.Add(Layer);
-	}
-	return Layer;
+	return Model.Get();
 }
 
 //--------------------------------------------------------------------------
 // UNearestNeighborOptimizedNetworkInstance
 //--------------------------------------------------------------------------
 
+UNearestNeighborOptimizedNetworkInstance::UNearestNeighborOptimizedNetworkInstance() = default;
+UNearestNeighborOptimizedNetworkInstance::UNearestNeighborOptimizedNetworkInstance(const FObjectInitializer& ObjectInitializer) : Super(ObjectInitializer) {}
+UNearestNeighborOptimizedNetworkInstance::UNearestNeighborOptimizedNetworkInstance(FVTableHelper& Helper) : Super(Helper) {}
+UNearestNeighborOptimizedNetworkInstance::~UNearestNeighborOptimizedNetworkInstance() = default;
+
 TArrayView<float> UNearestNeighborOptimizedNetworkInstance::GetInputs()
 { 
 	return TArrayView<float>(Inputs.GetData(), Inputs.Num());
@@ -136,50 +238,21 @@ void UNearestNeighborOptimizedNetworkInstance::Init(UNearestNeighborOptimizedNet
 	Inputs.SetNumZeroed(Network->GetNumInputs());
 	Outputs.SetNumZeroed(Network->GetNumOutputs());
 
-	int32 MaxNumUnits = 0;
-	const int32 NumLayers = Network->GetNumLayers();
-	for (int32 LayerIndex = 0; LayerIndex < NumLayers; ++LayerIndex)
+	if (Network->GetModel())
 	{
-		const UNearestNeighborNetworkLayer* CurLayer = Network->GetLayer(LayerIndex);
-		const int32 NumInputUnits = CurLayer->NumInputs;
-		const int32 NumOutputUnits = CurLayer->NumOutputs;
-		MaxNumUnits = FMath::Max3<int32>(NumInputUnits, NumOutputUnits, MaxNumUnits);
+		Instance = Network->GetModel()->CreateModelInstanceCPU();
+		Instance->SetInputTensorShapes({ UE::NNE::FTensorShape::Make({ 1, (uint32)Inputs.Num() }) });
 	}
-
-	TempInputArray.SetNumZeroed(MaxNumUnits);
-	TempOutputArray.SetNumZeroed(MaxNumUnits);
 }
 
-
 void UNearestNeighborOptimizedNetworkInstance::Run()
 {	
-	TRACE_CPUPROFILER_EVENT_SCOPE(UNearestNeighborOptimizedNetwork::Run)
-
-	// Setup the buffer pointers.
-	FRunSettings RunSettings;
-	RunSettings.TempInputBuffer  = TempInputArray.GetData();
-	RunSettings.TempOutputBuffer = TempOutputArray.GetData();
-	RunSettings.InputBuffer		 = Inputs.GetData();
-	RunSettings.OutputBuffer	 = Outputs.GetData();
-
-	// Run the network.
-	const int32 NumLayers = Network->GetNumLayers();
-	for (int32 LayerIndex = 0; LayerIndex < NumLayers; ++LayerIndex)
+	TRACE_CPUPROFILER_EVENT_SCOPE(UNearestNeighborOptimizedNetwork::Run);
+
+	if (Instance)
 	{
-		const UNearestNeighborNetworkLayer* CurLayer = Network->GetLayer(LayerIndex);
-		if (LayerIndex == 0)
-		{
-			const float* const RESTRICT NetworkInputs = RunSettings.InputBuffer;
-			CurLayer->Run(RunSettings.InputBuffer, RunSettings.TempInputBuffer);
-		}
-		else if (LayerIndex == NumLayers - 1)
-		{
-			CurLayer->Run(RunSettings.TempInputBuffer, RunSettings.OutputBuffer);
-		}
-		else
-		{
-			CurLayer->Run(RunSettings.TempInputBuffer, RunSettings.TempOutputBuffer);
-			Swap(RunSettings.TempInputBuffer, RunSettings.TempOutputBuffer);
-		}
+		Instance->RunSync(
+			{ { (void*)Inputs.GetData(), Inputs.Num() * sizeof(float) } },
+			{ { (void*)Outputs.GetData(), Outputs.Num() * sizeof(float) } });
 	}
 }
diff --git a/Engine/Plugins/Animation/MLDeformer/NearestNeighborModel/Source/NearestNeighborModel/Private/NearestNeighborOptimizedNetwork.h b/Engine/Plugins/Animation/MLDeformer/NearestNeighborModel/Source/NearestNeighborModel/Private/NearestNeighborOptimizedNetwork.h
--- a/Engine/Plugins/Animation/MLDeformer/NearestNeighborModel/Source/NearestNeighborModel/Private/NearestNeighborOptimizedNetwork.h
+++ b/Engine/Plugins/Animation/MLDeformer/NearestNeighborModel/Source/NearestNeighborModel/Private/NearestNeighborOptimizedNetwork.h
@@ -5,15 +5,18 @@
 #include "UObject/Object.h"
 #include "NearestNeighborOptimizedNetwork.generated.h"
 
-class UNearestNeighborOptimizedNetworkInstance;
+namespace UE::NNE
+{
+	class IModelCPU;
+	class IModelInstanceCPU;
+}
 
-#define NEARESTNEIGHBORMODEL_USE_ISPC INTEL_ISPC
-#if !defined(NEARESTNEIGHBORMODEL_USE_ISPC)
-	#define NEARESTNEIGHBORMODEL_USE_ISPC 0
-#endif
+class UNNEModelData;
+
+class UNearestNeighborOptimizedNetworkInstance;
 
 USTRUCT()
-struct FNearestNeighborNetworkParameter
+struct UE_DEPRECATED(5.4, "Nearest Neighbor Network Parameter no longer used for inference.") FNearestNeighborNetworkParameter
 {
 	GENERATED_BODY()
 
@@ -24,62 +27,43 @@ struct FNearestNeighborNetworkParameter
 	TArray<int32> Shape;
 };
 
-UENUM(BlueprintType)
-enum class ENearestNeighborNetworkLayerType : uint8
-{
-	None,
-	Gemm_Prelu,
-	Gemm,
-};
-
-
-/** A general network layer that contains a list of parameters. The Run() method should be implemented by child classes */ 
 UCLASS()
-class UNearestNeighborNetworkLayer
+class UE_DEPRECATED(5.4, "Nearest Neighbor Network Layer no longer used for inference.") UNearestNeighborNetworkLayer
 	: public UObject
 {
 	GENERATED_BODY()
 
 public:
-	/** The weight matrix number of inputs (rows). */
 	UPROPERTY(EditAnywhere, BlueprintReadWrite, Category = "MLDeformer")
 	int32 NumInputs = 0;
 
-	/** The weight matrix number of outputs (columns). */
 	UPROPERTY(EditAnywhere, BlueprintReadWrite, Category = "MLDeformer")
 	int32 NumOutputs = 0;
 
-	// /** The parameters of the layer */
+	PRAGMA_DISABLE_DEPRECATION_WARNINGS
+
 	UPROPERTY()
 	TArray<FNearestNeighborNetworkParameter> Parameters;
 
-	UFUNCTION(BlueprintCallable, Category = "MLDeformer")
-	void AddParameter(const TArray<float>& Values, const TArray<int32>& Shape);
-
-	virtual void Run(const float* RESTRICT InputBuffer, float* RESTRICT OutputBuffer) const;
+	PRAGMA_ENABLE_DEPRECATION_WARNINGS
 };
 
-/** Gemm: GEneral Matrix Multiplication: Wx + b */
+PRAGMA_DISABLE_DEPRECATION_WARNINGS
+
 UCLASS()
-class UNearestNeighborNetworkLayer_Gemm_Prelu
-	: public UNearestNeighborNetworkLayer
+class UE_DEPRECATED(5.4, "Nearest Neighbor Network Layer no longer used for inference.") UNearestNeighborNetworkLayer_Gemm_Prelu : public UNearestNeighborNetworkLayer
 {
 	GENERATED_BODY()
-
-public:
-	virtual void Run(const float* RESTRICT InputBuffer, float* RESTRICT OutputBuffer) const override;
 };
 
 UCLASS()
-class UNearestNeighborNetworkLayer_Gemm
-	: public UNearestNeighborNetworkLayer
+class UE_DEPRECATED(5.4, "Nearest Neighbor Network Layer no longer used for inference.") UNearestNeighborNetworkLayer_Gemm : public UNearestNeighborNetworkLayer
 {
 	GENERATED_BODY()
-
-public:
-	virtual void Run(const float* RESTRICT InputBuffer, float* RESTRICT OutputBuffer) const override;
 };
 
+PRAGMA_ENABLE_DEPRECATION_WARNINGS
+
 /**
  * The specialized neural network for the MLDeformerModel.
  * This class is used to do inference at runtime at a higher performance than using UNeuralNetwork.
@@ -91,6 +75,9 @@ class UNearestNeighborOptimizedNetwork
 	GENERATED_BODY()
 
 public:
+
+	virtual void PostLoad() override;
+
 	/** Clear the network, getting rid of all layers. */
 	UFUNCTION(BlueprintCallable, Category = "MLDeformer")
 	virtual void Empty();
@@ -129,31 +116,51 @@ public:
 	 */
 	virtual int32 GetNumOutputs() const;
 
-
 	/**
-	 * Get the number of network layers.
-	 * @return The number of network layer.
+	 * Sets the number of inputs, which is the number of floats the network takes as input.
+	 * @param InNumInputs The number of input floats to the network.
 	 */
-	const int32 GetNumLayers() const;
+	virtual void SetNumInputs(int32 InNumInputs);
 
 	/**
-	 * Get a given network layer.
-	 * @return A pointer to the layer, which will contain the parameters.
+	 * Sets the number of outputs, which is the number of floats the network will output.
+	 * @param InNumOutputs The number of floats that the network outputs.
 	 */
-	UNearestNeighborNetworkLayer* GetLayer(int32 Index) const;
-
+	virtual void SetNumOutputs(int32 InNumOutputs);
 
 	/**
-	 * Add a network layer.
-	 * @return A pointer to the layer, which will contain the parameters.
+	 * Get the Model.
+	 * @return A pointer to the Model.
 	 */
-	UFUNCTION(BlueprintCallable, Category = "MLDeformer")
-	UNearestNeighborNetworkLayer* AddLayer(const int32 LayerType);
+	UE::NNE::IModelCPU* GetModel() const;
 
 protected:
-	/** The network layers */
+
+	/** The Model Data used for network */
+	UPROPERTY()
+	TObjectPtr<UNNEModelData> ModelData;
+
+	/** The Neural Network used. */
+	TSharedPtr<UE::NNE::IModelCPU> Model;
+
 	UPROPERTY()
-	TArray<TObjectPtr<UNearestNeighborNetworkLayer>> Layers;
+	int32 NumInputs = 0;
+
+	UPROPERTY()
+	int32 NumOutputs = 0;
+
+	PRAGMA_DISABLE_DEPRECATION_WARNINGS
+
+	UE_DEPRECATED(5.4, "Layers Object no longer available.")
+	UPROPERTY(meta = (DeprecatedProperty, DeprecationMessage = "Network format changed."))
+	TArray<TObjectPtr<UNearestNeighborNetworkLayer>> Layers_DEPRECATED;
+
+	PRAGMA_ENABLE_DEPRECATION_WARNINGS
+
+private:
+
+	static const TCHAR* RuntimeName;
+
 };
 
 /** 
@@ -170,6 +177,13 @@ class UNearestNeighborOptimizedNetworkInstance
 	GENERATED_BODY()
 
 public:
+
+	// These are required because of the use of TUniquePtr
+	UNearestNeighborOptimizedNetworkInstance();
+	UNearestNeighborOptimizedNetworkInstance(const FObjectInitializer& ObjectInitializer);
+	UNearestNeighborOptimizedNetworkInstance(FVTableHelper& Helper);
+	virtual ~UNearestNeighborOptimizedNetworkInstance();
+
 	/**
 	 * Get the network input buffer.
 	 * @return The array view of floats that represents the network inputs.
@@ -198,13 +212,6 @@ public:
 	virtual void Run();
 
 protected:
-	struct FRunSettings
-	{
-		float* TempInputBuffer;
-		float* TempOutputBuffer;
-		const float* InputBuffer;
-		float* OutputBuffer;
-	};
 
 	/**
  	 * Initialize this instance using a given network object.
@@ -220,11 +227,8 @@ protected:
 	/** The output values. */
 	TArray<float> Outputs;
 
-	/** A pre-allocated temp buffer for inputs. */
-	TArray<float> TempInputArray;
-
-	/** A pre-allocated temp buffer for outputs. */
-	TArray<float> TempOutputArray;
+	/** The instance data for this Network Instance. */
+	TSharedPtr<UE::NNE::IModelInstanceCPU> Instance;
 
 	/** The neural network this is an instance of. */
 	UPROPERTY(Transient)
diff --git a/Engine/Plugins/Animation/MLDeformer/NearestNeighborModel/Source/NearestNeighborModel/Private/NearestNeighborOptimizedNetwork.ispc b/Engine/Plugins/Animation/MLDeformer/NearestNeighborModel/Source/NearestNeighborModel/Private/NearestNeighborOptimizedNetwork.ispc
deleted file mode 100644
diff --git a/Engine/Plugins/Animation/MLDeformer/NearestNeighborModel/Source/NearestNeighborModel/Public/NearestNeighborModel.h b/Engine/Plugins/Animation/MLDeformer/NearestNeighborModel/Source/NearestNeighborModel/Public/NearestNeighborModel.h
--- a/Engine/Plugins/Animation/MLDeformer/NearestNeighborModel/Source/NearestNeighborModel/Public/NearestNeighborModel.h
+++ b/Engine/Plugins/Animation/MLDeformer/NearestNeighborModel/Source/NearestNeighborModel/Public/NearestNeighborModel.h
@@ -261,8 +261,12 @@ private:
 	UNearestNeighborOptimizedNetwork* GetOptimizedNetwork() const { return OptimizedNetwork.Get(); }
 	int32 GetOptimizedNetworkNumOutputs() const;
 	void SetOptimizedNetwork(UNearestNeighborOptimizedNetwork* InOptimizedNetwork);
+	
+	UE_DEPRECATED(5.4, "Onnx file no longer used for training network. Use LoadOptimizedNetworkFromFile instead.")
 	bool LoadOptimizedNetwork(const FString& OnnxPath);
 
+	bool LoadOptimizedNetworkFromFile(const FString& Filename);
+
 #if WITH_EDITORONLY_DATA
 	TObjectPtr<UAnimSequence> GetNearestNeighborSkeletons(int32 PartId);
 	const TObjectPtr<UAnimSequence> GetNearestNeighborSkeletons(int32 PartId) const;
@@ -347,16 +351,17 @@ private:
 #endif
 
 protected:
-#if WITH_EDITORONLY_DATA
+
 	UPROPERTY(VisibleAnywhere, BlueprintReadOnly, Category = "Training Settings")
 	int32 InputDim = 0;
 
-	UPROPERTY(EditAnywhere, BlueprintReadWrite, Category = "Training Settings")
-	TArray<int32> HiddenLayerDims;
-
 	UPROPERTY(VisibleAnywhere, BlueprintReadOnly, Category = "Training Settings")
 	int32 OutputDim = 0;
 
+#if WITH_EDITORONLY_DATA
+	UPROPERTY(EditAnywhere, BlueprintReadWrite, Category = "Training Settings")
+	TArray<int32> HiddenLayerDims;
+
 	UPROPERTY(EditAnywhere, BlueprintReadWrite, Category = "Training Settings", meta = (ClampMin = "1"))
 	int32 NumEpochs = 10000;
 
diff --git a/Engine/Plugins/Animation/MLDeformer/NearestNeighborModel/Source/NearestNeighborModelEditor/Private/NearestNeighborEditorModel.cpp b/Engine/Plugins/Animation/MLDeformer/NearestNeighborModel/Source/NearestNeighborModelEditor/Private/NearestNeighborEditorModel.cpp
--- a/Engine/Plugins/Animation/MLDeformer/NearestNeighborModel/Source/NearestNeighborModelEditor/Private/NearestNeighborEditorModel.cpp
+++ b/Engine/Plugins/Animation/MLDeformer/NearestNeighborModel/Source/NearestNeighborModelEditor/Private/NearestNeighborEditorModel.cpp
@@ -99,10 +99,10 @@ namespace UE::NearestNeighborModel
 		UNearestNeighborModel* NearestNeighborModel = GetNearestNeighborModel();
 		if (NearestNeighborModel)
 		{	
-			const FString OnnxFile = GetTrainedNetworkOnnxFile();
+			const FString File = GetNearestNeighborModel()->GetModelDir() / TEXT("NearestNeighborModel.ubnne");
 			if (NearestNeighborModel->ShouldUseOptimizedNetwork())
 			{
-				const bool bSuccess = NearestNeighborModel->LoadOptimizedNetwork(OnnxFile);
+				const bool bSuccess = NearestNeighborModel->LoadOptimizedNetworkFromFile(File);
 				if (bSuccess)
 				{
 					UNearestNeighborModelInstance* ModelInstance = static_cast<UNearestNeighborModelInstance*>(GetTestMLDeformerModelInstance());
diff --git a/Engine/Plugins/Animation/MLDeformer/NearestNeighborModel/Source/NearestNeighborModelEditor/Public/NearestNeighborEditorModel.h b/Engine/Plugins/Animation/MLDeformer/NearestNeighborModel/Source/NearestNeighborModelEditor/Public/NearestNeighborEditorModel.h
--- a/Engine/Plugins/Animation/MLDeformer/NearestNeighborModel/Source/NearestNeighborModelEditor/Public/NearestNeighborEditorModel.h
+++ b/Engine/Plugins/Animation/MLDeformer/NearestNeighborModel/Source/NearestNeighborModelEditor/Public/NearestNeighborEditorModel.h
@@ -48,6 +48,8 @@ namespace UE::NearestNeighborModel
 		virtual bool LoadTrainedNetwork() const override;
 		virtual void OnPropertyChanged(FPropertyChangedEvent& PropertyChangedEvent) override;
 		virtual void OnPostTraining(ETrainingResult TrainingResult, bool bUsePartiallyTrainedWhenAborted) override;
+		
+		UE_DEPRECATED(5.4, "Onnx file no longer used for training network.")
 		virtual FString GetTrainedNetworkOnnxFile() const override;
 		virtual int32 GetNumTrainingFrames() const override;
 		// ~END FMLDeformerGeomCacheEditorModel overrides.
