NearestNeighborModel: clean up CPU memory of cloth part data

diff --git a/Engine/Plugins/Animation/MLDeformer/NearestNeighborModel/Source/NearestNeighborModel/Public/NearestNeighborModel.h b/Engine/Plugins/Animation/MLDeformer/NearestNeighborModel/Source/NearestNeighborModel/Public/NearestNeighborModel.h
--- a/Engine/Plugins/Animation/MLDeformer/NearestNeighborModel/Source/NearestNeighborModel/Public/NearestNeighborModel.h
+++ b/Engine/Plugins/Animation/MLDeformer/NearestNeighborModel/Source/NearestNeighborModel/Public/NearestNeighborModel.h
@@ -107,10 +107,11 @@
 	UPROPERTY()
 	uint32 NumNeighbors = 0;
 
+#if WITH_EDITORONLY_DATA
 	/** Vertex indices for this cloth part */
 	UPROPERTY()
 	TArray<uint32> VertexMap;
-
+	
 	/** PCA basis for this cloth part. This is a flattened array of size (PCACoeffNum, NumVertices * 3)  */
 	UPROPERTY()
 	TArray<float> PCABasis;
@@ -119,13 +120,14 @@
 	UPROPERTY()
 	TArray<float> VertexMean;
 
-	/** PCA coefficients of the nearest neighbors. This is a flattened array of size (NumNeighbors, PCACoeffNum) */
-	UPROPERTY()
-	TArray<float> NeighborCoeffs;
-
 	/** The remaining offsets of the nearest neighbor shapes (after reducing PCA offsets). This is a flattened array of size (NumNeighbors, PCACoeffNum) */
 	UPROPERTY()
 	TArray<float> NeighborOffsets;
+#endif
+
+	/** PCA coefficients of the nearest neighbors. This is a flattened array of size (NumNeighbors, PCACoeffNum) */
+	UPROPERTY()
+	TArray<float> NeighborCoeffs;
 };
 
 namespace UE::NearestNeighborModel
@@ -204,6 +206,38 @@
 	float GetDecayFactor() const { return DecayFactor; };
 	float GetNearestNeighborOffsetWeight() const { return NearestNeighborOffsetWeight; }
 
+	UFUNCTION(BlueprintCallable, Category = "Nearest Neighbor Model")
+	void SetNumNeighbors(int32 PartId, int32 InNumNeighbors) { ClothPartData[PartId].NumNeighbors = InNumNeighbors; }
+	
+	UFUNCTION(BlueprintPure, Category = "Nearest Neighbor Model")
+	const TArray<float>& NeighborCoeffs(int32 PartId) const { return ClothPartData[PartId].NeighborCoeffs; }
+	
+	TArray<float>& NeighborCoeffs(int32 PartId) { return ClothPartData[PartId].NeighborCoeffs; }
+
+	UFUNCTION(BlueprintPure, Category = "Nearest Neighbor Model")
+	TArray<float> ClipInputs(const TArray<float>& Input) const;
+
+	void ClipInputs(float* InputPtr, int NumInputs) const;
+
+	void InitInputInfo();
+
+	void ResetMorphBuffers();
+
+	virtual int32 GetNumFloatsPerBone() const override { return NearestNeighborNumFloatsPerBone; }
+	virtual int32 GetNumFloatsPerCurve() const override { return NearestNeighborNumFloatsPerCurve; }
+
+	bool DoesUseOptimizedNetwork() const;
+	bool ShouldUseOptimizedNetwork() const;
+	UNearestNeighborOptimizedNetwork* GetOptimizedNetwork() const { return OptimizedNetwork.Get(); }
+	int32 GetOptimizedNetworkNumOutputs() const;
+	void SetOptimizedNetwork(UNearestNeighborOptimizedNetwork* InOptimizedNetwork);
+
+	UE_DEPRECATED(5.4, "Onnx file no longer used for training network. Use LoadOptimizedNetworkFromFile instead.")
+	bool LoadOptimizedNetwork(const FString& OnnxPath);
+
+	bool LoadOptimizedNetworkFromFile(const FString& Filename);
+
+#if WITH_EDITOR
 	const TArray<uint32>& PartVertexMap(int32 PartId) const { return ClothPartData[PartId].VertexMap; }
 
 	UFUNCTION(BlueprintPure, Category = "Nearest Neighbor Model")
@@ -223,13 +257,5 @@
 	void SetVertexMean(int32 PartId, const TArray<float>& VertexMean) { ClothPartData[PartId].VertexMean = VertexMean; }
 
 	UFUNCTION(BlueprintCallable, Category = "Nearest Neighbor Model")
-	void SetNumNeighbors(int32 PartId, int32 InNumNeighbors) { ClothPartData[PartId].NumNeighbors = InNumNeighbors; }
-	
-	UFUNCTION(BlueprintPure, Category = "Nearest Neighbor Model")
-	const TArray<float>& NeighborCoeffs(int32 PartId) const { return ClothPartData[PartId].NeighborCoeffs; }
-
-	TArray<float>& NeighborCoeffs(int32 PartId) { return ClothPartData[PartId].NeighborCoeffs; }
-
-	UFUNCTION(BlueprintCallable, Category = "Nearest Neighbor Model")
 	void SetNeighborCoeffs(int32 PartId, const TArray<float>& NeighborCoeffs) { ClothPartData[PartId].NeighborCoeffs = NeighborCoeffs; }
 
@@ -241,31 +267,7 @@
 
 	UFUNCTION(BlueprintCallable, Category = "Nearest Neighbor Model")
 	void SetNeighborOffsets(int32 PartId, const TArray<float>& NeighborOffsets) { ClothPartData[PartId].NeighborOffsets = NeighborOffsets; }
+#endif
 
-	UFUNCTION(BlueprintPure, Category = "Nearest Neighbor Model")
-	TArray<float> ClipInputs(const TArray<float>& Input) const;
-
-	void ClipInputs(float* InputPtr, int NumInputs) const;
-
-	bool CheckPCAData(int32 PartId) const;
-
-	void InitInputInfo();
-
-	void ResetMorphBuffers();
-
-	virtual int32 GetNumFloatsPerBone() const override { return NearestNeighborNumFloatsPerBone; }
-	virtual int32 GetNumFloatsPerCurve() const override { return NearestNeighborNumFloatsPerCurve; }
-
-	bool DoesUseOptimizedNetwork() const;
-	bool ShouldUseOptimizedNetwork() const;
-	UNearestNeighborOptimizedNetwork* GetOptimizedNetwork() const { return OptimizedNetwork.Get(); }
-	int32 GetOptimizedNetworkNumOutputs() const;
-	void SetOptimizedNetwork(UNearestNeighborOptimizedNetwork* InOptimizedNetwork);
-	
-	UE_DEPRECATED(5.4, "Onnx file no longer used for training network. Use LoadOptimizedNetworkFromFile instead.")
-	bool LoadOptimizedNetwork(const FString& OnnxPath);
-
-	bool LoadOptimizedNetworkFromFile(const FString& Filename);
-
 #if WITH_EDITORONLY_DATA
 	TObjectPtr<UAnimSequence> GetNearestNeighborSkeletons(int32 PartId);
@@ -277,6 +279,7 @@
 
 	void UpdateNetworkInputDim();
 	void UpdateNetworkOutputDim();
+	bool CheckPCAData(int32 PartId) const;
 	UE::NearestNeighborModel::EUpdateResult UpdateClothPartData();
 	UE::NearestNeighborModel::EUpdateResult UpdateVertexMap(int32 PartId, const FString& VertexMapPath, const FSkelMeshImportedMeshInfo& Info);
 	void UpdatePCACoeffNums();
