NearestNeighborModel: add a simple tool to prune poses

diff --git a/Engine/Plugins/Animation/MLDeformer/NearestNeighborModel/Source/NearestNeighborModelEditor/Private/NearestNeighborEditorModel.cpp b/Engine/Plugins/Animation/MLDeformer/NearestNeighborModel/Source/NearestNeighborModelEditor/Private/NearestNeighborEditorModel.cpp
--- a/Engine/Plugins/Animation/MLDeformer/NearestNeighborModel/Source/NearestNeighborModelEditor/Private/NearestNeighborEditorModel.cpp
+++ b/Engine/Plugins/Animation/MLDeformer/NearestNeighborModel/Source/NearestNeighborModelEditor/Private/NearestNeighborEditorModel.cpp
@@ -934,6 +934,12 @@ namespace UE::NearestNeighborModel
 		}
 		return ModelInstance;
 	}
+
+	void FNearestNeighborEditorModel::GetNeighborStats()
+	{
+		UNearestNeighborTrainingModel *TrainingModel = InitTrainingModel<UNearestNeighborTrainingModel>(this);
+		TrainingModel->GetNeighborStats(0);
+	}
 }	// namespace UE::NearestNeighborModel
 
 #undef LOCTEXT_NAMESPACE
diff --git a/Engine/Plugins/Animation/MLDeformer/NearestNeighborModel/Source/NearestNeighborModelEditor/Private/NearestNeighborGeomCacheSampler.cpp b/Engine/Plugins/Animation/MLDeformer/NearestNeighborModel/Source/NearestNeighborModelEditor/Private/NearestNeighborGeomCacheSampler.cpp
--- a/Engine/Plugins/Animation/MLDeformer/NearestNeighborModel/Source/NearestNeighborModelEditor/Private/NearestNeighborGeomCacheSampler.cpp
+++ b/Engine/Plugins/Animation/MLDeformer/NearestNeighborModel/Source/NearestNeighborModelEditor/Private/NearestNeighborGeomCacheSampler.cpp
@@ -18,6 +18,8 @@
 using namespace UE::MLDeformer;
 namespace UE::NearestNeighborModel
 {
+	FNearestNeighborGeomCacheSampler::~FNearestNeighborGeomCacheSampler() = default;
+	
 	void FNearestNeighborGeomCacheSampler::Sample(int32 InAnimFrameIndex)
 	{
 		UNearestNeighborModel* NearestNeighborModel = static_cast<UNearestNeighborModel*>(Model);
@@ -210,19 +212,17 @@ namespace UE::NearestNeighborModel
 		}
 	}
 
-	bool FNearestNeighborGeomCacheSampler::SampleKMeansAnim(const int32 AnimId)
+	bool FNearestNeighborGeomCacheSampler::SetAnimToSample(UAnimSequence& InAnimToSample)
 	{
-		UNearestNeighborModel* NearestNeighborModel = static_cast<UNearestNeighborModel*>(Model);
-		if (NearestNeighborModel && AnimId < NearestNeighborModel->SourceAnims.Num() && SkeletalMeshComponent)
+		if (SkeletalMeshComponent)
 		{
-			const TObjectPtr<UAnimSequence> AnimSequence = NearestNeighborModel->SourceAnims[AnimId];
 			SkeletalMeshComponent->SetAnimationMode(EAnimationMode::AnimationSingleNode);
-			SkeletalMeshComponent->SetAnimation(AnimSequence);
+			SkeletalMeshComponent->SetAnimation(&InAnimToSample);
 			SkeletalMeshComponent->SetPosition(0.0f);
 			SkeletalMeshComponent->SetPlayRate(1.0f);
 			SkeletalMeshComponent->Play(false);
 			SkeletalMeshComponent->RefreshBoneTransforms();
-			KMeansAnimId = AnimId;
+			AnimToSample = &InAnimToSample;
 			return true;
 		}
 		else
@@ -231,53 +231,63 @@ namespace UE::NearestNeighborModel
 		}
 	}
 
-	
+	int32 FNearestNeighborGeomCacheSampler::GetAnimNumFrames() const
+	{
+		if (AnimToSample)
+		{
+			return AnimToSample->GetDataModel()->GetNumberOfKeys();
+		}
+		else
+		{
+			return 0;
+		}
+	}
 
-	bool FNearestNeighborGeomCacheSampler::SampleKMeansFrame(const int32 Frame)
+	bool FNearestNeighborGeomCacheSampler::SampleAnim(int32 Frame)
 	{
-		const USkeletalMesh* SkeletalMesh = SkeletalMeshComponent->GetSkeletalMeshAsset();
-		if (SkeletalMeshComponent && SkeletalMesh)
+		if (!SkeletalMeshComponent || !AnimToSample)
 		{
-			UNearestNeighborModel* NearestNeighborModel = static_cast<UNearestNeighborModel*>(Model);
-			const UAnimSequence* AnimSequence = NearestNeighborModel->SourceAnims[KMeansAnimId];
-			if (NearestNeighborModel->GetSkeletalMesh() == nullptr)
-			{
-				UE_LOG(LogNearestNeighborModel, Error, TEXT("SkeletalMesh is nullptr. Unable to sample KMeans frame."));
-				return false;
-			}
+			return false;
+		}
 
-			if (AnimSequence)
+		const int32 NumKeys = AnimToSample->GetDataModel()->GetNumberOfKeys();
+		if (Frame < 0 || Frame >= NumKeys)
+		{
+			UE_LOG(LogNearestNeighborModel, Warning, TEXT("AnimSequence only has %d keys, but being sampled with key %d"), NumKeys, Frame);
+			return false;
+		}
+
+		AnimFrameIndex = Frame;
+		SampleTime = GetTimeAtFrame(Frame);
+		
+		UpdateSkeletalMeshComponent();
+		UpdateBoneRotations();
+		UpdateCurveValues();
+		return true;
+	}
+
+	bool FNearestNeighborGeomCacheSampler::SampleKMeansAnim(const int32 AnimId)
+	{
+		UNearestNeighborModel* NearestNeighborModel = static_cast<UNearestNeighborModel*>(Model);
+		if (NearestNeighborModel && AnimId < NearestNeighborModel->SourceAnims.Num() && SkeletalMeshComponent)
+		{
+			if (const TObjectPtr<UAnimSequence> AnimSequence = NearestNeighborModel->SourceAnims[AnimId])
 			{
-				if (Frame < AnimSequence->GetDataModel()->GetNumberOfKeys())
+				if (SetAnimToSample(*AnimSequence))
 				{
-					AnimFrameIndex = Frame;
-					SampleTime = GetTimeAtFrame(Frame);
-
-					UpdateSkeletalMeshComponent();
-					UpdateBoneRotations();
-					UpdateCurveValues();
+					KMeansAnimId = AnimId;
 					return true;
 				}
-				else
-				{
-					UE_LOG(LogNearestNeighborModel, Error, TEXT("AnimSequence only has %d keys, but being sampled with key %d"), AnimSequence->GetDataModel()->GetNumberOfKeys(), Frame);
-					return false;
-				}
 			}
-			else
-			{
-				UE_LOG(LogNearestNeighborModel, Error, TEXT("AnimSequence %d is nullptr. Unable to sample KMeans frame."), KMeansAnimId);
-				return false;
-			}
-
-		}
-		else
-		{
-			UE_LOG(LogNearestNeighborModel, Error, TEXT("KMeans: SkeletalMesh does not exist"));
 		}
 		return false;
 	}
 
+	bool FNearestNeighborGeomCacheSampler::SampleKMeansFrame(const int32 Frame)
+	{
+		return SampleAnim(Frame);
+	}
+
 	TArray<uint32> FNearestNeighborGeomCacheSampler::GetMeshIndexBuffer() const
 	{
 		TArray<uint32> IndexBuffer;
diff --git a/Engine/Plugins/Animation/MLDeformer/NearestNeighborModel/Source/NearestNeighborModelEditor/Private/NearestNeighborModelVizSettingsDetails.cpp b/Engine/Plugins/Animation/MLDeformer/NearestNeighborModel/Source/NearestNeighborModelEditor/Private/NearestNeighborModelVizSettingsDetails.cpp
--- a/Engine/Plugins/Animation/MLDeformer/NearestNeighborModel/Source/NearestNeighborModelEditor/Private/NearestNeighborModelVizSettingsDetails.cpp
+++ b/Engine/Plugins/Animation/MLDeformer/NearestNeighborModel/Source/NearestNeighborModelEditor/Private/NearestNeighborModelVizSettingsDetails.cpp
@@ -7,6 +7,7 @@
 #include "DetailWidgetRow.h"
 #include "IDetailsView.h"
 #include "IDetailGroup.h"
+#include "NearestNeighborEditorModel.h"
 #include "NearestNeighborModelVizSettings.h"
 
 #define LOCTEXT_NAMESPACE "MLDeformerMorphModelVizSettingsDetails"
@@ -19,6 +20,28 @@ namespace UE::NearestNeighborModel
 		IDetailGroup& NNGroup = LiveSettingsCategory->AddGroup("Nearest Neighbor", LOCTEXT("NearestNeighborLabel", "Nearest Neighbor"), false, true);
 		NNGroup.AddPropertyRow(DetailLayoutBuilder->GetProperty(UNearestNeighborModelVizSettings::GetNearestNeighborActorsOffsetPropertyName(), UNearestNeighborModelVizSettings::StaticClass()));
 		NNGroup.AddPropertyRow(DetailLayoutBuilder->GetProperty(UNearestNeighborModelVizSettings::GetNearestNeighborIdsPropertyName(), UNearestNeighborModelVizSettings::StaticClass()));
+
+		IDetailGroup& ToolsGroup = LiveSettingsCategory->AddGroup("Tools", LOCTEXT("ToolsLabel", "Tools"), false, true);
+		ToolsGroup.AddWidgetRow()
+		.WholeRowContent()
+		[
+			SNew(SButton)
+			.Text(FText::FromString("Get Neighbor Stats"))
+			.HAlign(HAlign_Center)
+			.OnClicked_Lambda([this]
+			{
+				if (GetCastEditorModel())
+				{
+					GetCastEditorModel()->GetNeighborStats();
+				}
+				return FReply::Handled();
+			})
+		];
+	}
+
+	FNearestNeighborEditorModel* FNearestNeighborModelVizSettingsDetails::GetCastEditorModel()
+	{
+		return static_cast<FNearestNeighborEditorModel*>(EditorModel);
 	}
 
 };
diff --git a/Engine/Plugins/Animation/MLDeformer/NearestNeighborModel/Source/NearestNeighborModelEditor/Private/NearestNeighborTrainingModel.cpp b/Engine/Plugins/Animation/MLDeformer/NearestNeighborModel/Source/NearestNeighborModelEditor/Private/NearestNeighborTrainingModel.cpp
--- a/Engine/Plugins/Animation/MLDeformer/NearestNeighborModel/Source/NearestNeighborModelEditor/Private/NearestNeighborTrainingModel.cpp
+++ b/Engine/Plugins/Animation/MLDeformer/NearestNeighborModel/Source/NearestNeighborModelEditor/Private/NearestNeighborTrainingModel.cpp
@@ -1,15 +1,17 @@
 // Copyright Epic Games, Inc. All Rights Reserved.
 
 #include "NearestNeighborTrainingModel.h"
+#include "NearestNeighborGeomCacheSampler.h"
 #include "NearestNeighborModel.h"
 #include "NearestNeighborModelInstance.h"
 #include "NearestNeighborEditorModel.h"
-#include "NearestNeighborGeomCacheSampler.h"
 #include "Animation/AnimSequence.h"
 
 #define LOCTEXT_NAMESPACE "NearestNeighborTrainingModel"
 using namespace UE::NearestNeighborModel;
 
+UNearestNeighborTrainingModel::~UNearestNeighborTrainingModel() = default;
+
 void UNearestNeighborTrainingModel::Init(UE::MLDeformer::FMLDeformerEditorModel* InEditorModel)
 {
 	UMLDeformerTrainingModel::Init(InEditorModel);
@@ -44,36 +46,59 @@ int32 UNearestNeighborTrainingModel::GetPartNumNeighbors(const int32 PartId) con
 	return FMath::Min(NearestNeighborModel->GetNumNeighborsFromAnimSequence(PartId), NearestNeighborModel->GetNumNeighborsFromGeometryCache(PartId));
 }
 
-bool UNearestNeighborTrainingModel::SampleKmeansAnim(const int32 SkeletonId)
+bool UNearestNeighborTrainingModel::SetAnimToSample(UAnimSequence* AnimToSample)
 {
-	FMLDeformerSampler* BaseSampler = EditorModel->GetNumTrainingInputAnims() > 0 ? EditorModel->GetSamplerForTrainingAnim(0) : nullptr;
-	FNearestNeighborGeomCacheSampler* Sampler = static_cast<FNearestNeighborGeomCacheSampler*>(BaseSampler);
-	if (Sampler == nullptr)
+	if (!NearestNeighborModel || !AnimToSample)
 	{
 		return false;
 	}
-
-	return Sampler->SampleKMeansAnim(SkeletonId);
+	if (!AnimSampler.Get())
+	{
+		SetNewAnimSampler();
+		check(AnimSampler.Get());
+	}
+	static_cast<FNearestNeighborGeomCacheSampler*>(AnimSampler.Get())->SetAnimToSample(*AnimToSample);
+	return true;
 }
 
-bool UNearestNeighborTrainingModel::SampleKmeansFrame(const int32 Frame)
+
+bool UNearestNeighborTrainingModel::SampleAnim(int32 Frame)
 {
-	FMLDeformerSampler* BaseSampler = EditorModel->GetNumTrainingInputAnims() > 0 ? EditorModel->GetSamplerForTrainingAnim(0) : nullptr;
-	FNearestNeighborGeomCacheSampler* Sampler = static_cast<FNearestNeighborGeomCacheSampler*>(BaseSampler);
-	if (Sampler == nullptr)
+	if (!AnimSampler.Get())
 	{
 		return false;
 	}
+	if (!static_cast<FNearestNeighborGeomCacheSampler*>(AnimSampler.Get())->SampleAnim(Frame))
+	{
+		return false;
+	}
+	SampleBoneRotations = AnimSampler->GetBoneRotations();
+	return true;
+}
 
-	const bool bSampleExist = Sampler->SampleKMeansFrame(Frame);
-	if (bSampleExist)
+int32 UNearestNeighborTrainingModel::GetAnimNumFrames() const
+{
+	if (!AnimSampler.Get())
 	{
-		SampleBoneRotations = Sampler->GetBoneRotations();
-		return true;	
+		return 0;
 	}
-	return false;
+	return static_cast<FNearestNeighborGeomCacheSampler*>(AnimSampler.Get())->GetAnimNumFrames();
 }
 
+bool UNearestNeighborTrainingModel::SampleKmeansAnim(const int32 SkeletonId)
+{
+	if (!AnimSampler.Get())
+	{
+		SetNewAnimSampler();
+		check(AnimSampler.Get());
+	}
+	return static_cast<FNearestNeighborGeomCacheSampler*>(AnimSampler.Get())->SampleKMeansAnim(SkeletonId);
+}
+
+bool UNearestNeighborTrainingModel::SampleKmeansFrame(const int32 Frame)
+{
+	return SampleAnim(Frame);
+}
 
 int32 UNearestNeighborTrainingModel::GetKmeansNumAnims() const
 {
@@ -142,5 +167,26 @@ void UNearestNeighborTrainingModel::DestroyModelInstance(UNearestNeighborModelIn
 	ModelInstance->ConditionalBeginDestroy();
 }
 
+void UNearestNeighborTrainingModel::SetNewAnimSampler()
+{
+	AnimSampler = MakeUnique<FNearestNeighborGeomCacheSampler>();
+	AnimSampler->Init(EditorModel, 0);
+}
+
+const UAnimSequence* UNearestNeighborTrainingModel::GetTestAnim() const
+{
+	if (!EditorModel)
+	{
+		return nullptr;
+	}
+	const FNearestNeighborEditorModel* NNEditorModel = static_cast<FNearestNeighborEditorModel*>(EditorModel);
+	const UNearestNeighborModelVizSettings* Viz = NNEditorModel->GetNearestNeighborModelVizSettings();
+	if (!Viz)
+	{
+		return nullptr;
+	}
+	return Viz->GetTestAnimSequence();
+}
+
 
 #undef LOCTEXT_NAMESPACE
diff --git a/Engine/Plugins/Animation/MLDeformer/NearestNeighborModel/Source/NearestNeighborModelEditor/Public/NearestNeighborEditorModel.h b/Engine/Plugins/Animation/MLDeformer/NearestNeighborModel/Source/NearestNeighborModelEditor/Public/NearestNeighborEditorModel.h
--- a/Engine/Plugins/Animation/MLDeformer/NearestNeighborModel/Source/NearestNeighborModelEditor/Public/NearestNeighborEditorModel.h
+++ b/Engine/Plugins/Animation/MLDeformer/NearestNeighborModel/Source/NearestNeighborModelEditor/Public/NearestNeighborEditorModel.h
@@ -23,6 +23,7 @@ namespace UE::NearestNeighborModel
 	class FNearestNeighborEditorModelActor;
 	class FNearestNeighborModelDetails;
 	class FNearestNeighborModelSampler;
+	class FNearestNeighborModelVizSettingsDetails;
 
 	class NEARESTNEIGHBORMODELEDITOR_API FNearestNeighborEditorModel
 		: public UE::MLDeformer::FMLDeformerMorphModelEditorModel
@@ -57,6 +58,7 @@ namespace UE::NearestNeighborModel
 		friend class FNearestNeighborModelSampler;
 		friend class FNearestNeighborModelDetails;
 		friend class ::UNearestNeighborTrainingModel;
+		friend class FNearestNeighborModelVizSettingsDetails;
 
 	private:
 		// Some helpers that cast to this model's variants of some classes.
@@ -107,6 +109,8 @@ namespace UE::NearestNeighborModel
 			return TrainingModel;
 		}
 
+		void GetNeighborStats();
+
 		// Actors showing the nearest neighbors of the current frame
 		TArray<FNearestNeighborEditorModelActor*> NearestNeighborActors;
 
diff --git a/Engine/Plugins/Animation/MLDeformer/NearestNeighborModel/Source/NearestNeighborModelEditor/Public/NearestNeighborGeomCacheSampler.h b/Engine/Plugins/Animation/MLDeformer/NearestNeighborModel/Source/NearestNeighborModelEditor/Public/NearestNeighborGeomCacheSampler.h
--- a/Engine/Plugins/Animation/MLDeformer/NearestNeighborModel/Source/NearestNeighborModelEditor/Public/NearestNeighborGeomCacheSampler.h
+++ b/Engine/Plugins/Animation/MLDeformer/NearestNeighborModel/Source/NearestNeighborModelEditor/Public/NearestNeighborGeomCacheSampler.h
@@ -4,6 +4,7 @@
 
 #include "MLDeformerGeomCacheSampler.h"
 
+class UAnimSequence;
 class UNearestNeighborTrainingModel;
 namespace UE::NearestNeighborModel
 {
@@ -13,6 +14,7 @@ namespace UE::NearestNeighborModel
 		: public UE::MLDeformer::FMLDeformerGeomCacheSampler
 	{
 	public:
+		virtual ~FNearestNeighborGeomCacheSampler();
 		// FMLDeformerGeomCacheSampler overrides
 		virtual void Sample(int32 InAnimFrameIndex) override;
 		// ~END FMLDeformerGeomCacheSampler overrides
@@ -24,11 +26,16 @@ namespace UE::NearestNeighborModel
 		virtual void SampleDualQuaternionDeltas(int32 InAnimFrameIndex);
 		virtual bool SampleKMeansAnim(const int32 SkeletonId);
 		virtual bool SampleKMeansFrame(const int32 Frame);
+		int32 GetAnimNumFrames() const;
+		bool SetAnimToSample(UAnimSequence& InAnimToSample);
+		bool SampleAnim(int32 Frame);
 		uint8 GenerateMeshMappings();
 		uint8 CheckMeshMappingsEmpty() const;
 		TArray<uint32> GetMeshIndexBuffer() const;
 		FVector3f CalcDualQuaternionDelta(int32 VertexIndex, const FVector3f& WorldDelta, const FSkeletalMeshLODRenderData& SkelMeshLODData, const FSkinWeightVertexBuffer& SkinWeightBuffer) const;
 
 		int32 KMeansAnimId;
+
+		TObjectPtr<UAnimSequence> AnimToSample;
 	};
 }
diff --git a/Engine/Plugins/Animation/MLDeformer/NearestNeighborModel/Source/NearestNeighborModelEditor/Public/NearestNeighborModelVizSettingsDetails.h b/Engine/Plugins/Animation/MLDeformer/NearestNeighborModel/Source/NearestNeighborModelEditor/Public/NearestNeighborModelVizSettingsDetails.h
--- a/Engine/Plugins/Animation/MLDeformer/NearestNeighborModel/Source/NearestNeighborModelEditor/Public/NearestNeighborModelVizSettingsDetails.h
+++ b/Engine/Plugins/Animation/MLDeformer/NearestNeighborModel/Source/NearestNeighborModelEditor/Public/NearestNeighborModelVizSettingsDetails.h
@@ -7,6 +7,8 @@
 
 namespace UE::NearestNeighborModel
 {
+	class FNearestNeighborEditorModel;
+
 	class NEARESTNEIGHBORMODELEDITOR_API FNearestNeighborModelVizSettingsDetails
 		: public UE::MLDeformer::FMLDeformerMorphModelVizSettingsDetails
 	{
@@ -18,5 +20,6 @@ namespace UE::NearestNeighborModel
 		}
 		
 		virtual void AddAdditionalSettings() override;
+		FNearestNeighborEditorModel* GetCastEditorModel();
 	};
 }	// namespace UE::NearestNeighborModel
diff --git a/Engine/Plugins/Animation/MLDeformer/NearestNeighborModel/Source/NearestNeighborModelEditor/Public/NearestNeighborTrainingModel.h b/Engine/Plugins/Animation/MLDeformer/NearestNeighborModel/Source/NearestNeighborModelEditor/Public/NearestNeighborTrainingModel.h
--- a/Engine/Plugins/Animation/MLDeformer/NearestNeighborModel/Source/NearestNeighborModelEditor/Public/NearestNeighborTrainingModel.h
+++ b/Engine/Plugins/Animation/MLDeformer/NearestNeighborModel/Source/NearestNeighborModelEditor/Public/NearestNeighborTrainingModel.h
@@ -6,11 +6,18 @@
 #include "MLDeformerGeomCacheTrainingModel.h"
 #include "NearestNeighborTrainingModel.generated.h"
 
+
+namespace UE::MLDeformer
+{
+	class FMLDeformerGeomCacheSampler;
+};
+
 namespace UE::NearestNeighborModel
 {
 	class FNearestNeighborEditorModel;
-}
+};
 
+class UAnimSequence;
 class UNearestNeighborModel;
 class UNearestNeighborModelInstance;
 
@@ -21,6 +28,7 @@ class NEARESTNEIGHBORMODELEDITOR_API UNearestNeighborTrainingModel
 	GENERATED_BODY()
 
 public:
+	virtual ~UNearestNeighborTrainingModel();
 	virtual void Init(UE::MLDeformer::FMLDeformerEditorModel* InEditorModel) override;
 
 	/** Main training function, with implementation in python. */
@@ -36,6 +44,9 @@ protected:
 	UFUNCTION(BlueprintImplementableEvent, Category = "Python")
 	int32 KmeansClusterPoses(const int32 PartId) const;
 
+	UFUNCTION(BlueprintImplementableEvent, Category = "Python")
+	bool GetNeighborStats(int32 PartId);
+
 	UPROPERTY(EditAnywhere, BlueprintReadWrite, Category = "Training Data")
 	TArray<float> PartSampleDeltas;
 
@@ -57,6 +68,15 @@ private:
 	UFUNCTION(BlueprintPure, Category = "Python")
 	int32 GetPartNumNeighbors(const int32 PartId) const;
 
+	UFUNCTION(BlueprintCallable, Category = "Python")
+	bool SetAnimToSample(UAnimSequence* AnimToSample);
+
+	UFUNCTION(BlueprintCallable, Category = "Python")
+	bool SampleAnim(int32 Frame);
+
+	UFUNCTION(BlueprintPure, Category = "Python")
+	int32 GetAnimNumFrames() const;
+
 	UFUNCTION(BlueprintCallable, Category = "Python")
 	bool SampleKmeansAnim(const int32 SkeletonId);
 
@@ -84,5 +104,12 @@ private:
 	UFUNCTION(BlueprintCallable, Category = "Python")
 	void DestroyModelInstance(UNearestNeighborModelInstance* ModelInstance);
 
+	UFUNCTION(BlueprintCallable, Category = "Python")
+	const UAnimSequence* GetTestAnim() const;
+
+	void SetNewAnimSampler();
+
 	UNearestNeighborModel* NearestNeighborModel = nullptr;
+
+	TUniquePtr<UE::MLDeformer::FMLDeformerGeomCacheSampler> AnimSampler;
 };
