NearestNeighborModel: add exclude frame
fix confusing namespaces

diff --git a/Engine/Plugins/Animation/MLDeformer/NearestNeighborModel/Source/NearestNeighborModel/Private/NearestNeighborModel.cpp b/Engine/Plugins/Animation/MLDeformer/NearestNeighborModel/Source/NearestNeighborModel/Private/NearestNeighborModel.cpp
--- a/Engine/Plugins/Animation/MLDeformer/NearestNeighborModel/Source/NearestNeighborModel/Private/NearestNeighborModel.cpp
+++ b/Engine/Plugins/Animation/MLDeformer/NearestNeighborModel/Source/NearestNeighborModel/Private/NearestNeighborModel.cpp
@@ -73,8 +73,8 @@ void UNearestNeighborModel::UpdateMemoryUsage()
 	for (const FClothPartData& Part : ClothPartData)
 	{
 		ClothDataSize += Part.PCABasis.GetTypeSize() * Part.PCABasis.Num();
-		ClothDataSize += Part.NeighborOffsets.GetTypeSize() * Part.NeighborOffsets.Num();
-		ClothDataSize += Part.NeighborCoeffs.GetTypeSize() * Part.NeighborCoeffs.Num();
+		ClothDataSize += Part.AssetNeighborOffsets.GetTypeSize() * Part.AssetNeighborOffsets.Num();
+		ClothDataSize += Part.AssetNeighborCoeffs.GetTypeSize() * Part.AssetNeighborCoeffs.Num();
 		ClothDataSize += Part.VertexMean.GetTypeSize() * Part.VertexMean.Num();
 		ClothDataSize += Part.VertexMap.GetTypeSize() * Part.VertexMap.Num();
 	}
@@ -141,6 +141,13 @@ TArray<uint32> AddConstant(const TArray<uint32> &InArr, uint32 Constant)
 	return OutArr;
 }
 
+#if WITH_EDITOR
+int32 UNearestNeighborModel::GetAssetNumNeighbors(int32 PartId) const
+{
+	return FMath::Min(GetNumNeighborsFromGeometryCache(PartId), GetNumNeighborsFromAnimSequence(PartId));
+}
+#endif
+
 TArray<float> UNearestNeighborModel::ClipInputs(const TArray<float>& Inputs) const
 {
 	TArray<float> Result = Inputs;
@@ -331,8 +338,8 @@ UE::NearestNeighborModel::EUpdateResult UNearestNeighborModel::UpdateClothPartDa
 			ClothPartData[PartId].PCABasis.SetNumZeroed(ClothPartData[PartId].NumVertices * 3 * ClothPartData[PartId].PCACoeffNum);
 
 			// Init default neighbor data.
-			ClothPartData[PartId].NeighborCoeffs.SetNumZeroed(ClothPartData[PartId].PCACoeffNum);
-			ClothPartData[PartId].NeighborOffsets.SetNumZeroed(ClothPartData[PartId].NumVertices * 3);
+			ClothPartData[PartId].AssetNeighborCoeffs.SetNumZeroed(ClothPartData[PartId].PCACoeffNum);
+			ClothPartData[PartId].AssetNeighborOffsets.SetNumZeroed(ClothPartData[PartId].NumVertices * 3);
 			ClothPartData[PartId].NumNeighbors = 1;
 		}
 	}
diff --git a/Engine/Plugins/Animation/MLDeformer/NearestNeighborModel/Source/NearestNeighborModel/Public/NearestNeighborModel.h b/Engine/Plugins/Animation/MLDeformer/NearestNeighborModel/Source/NearestNeighborModel/Public/NearestNeighborModel.h
--- a/Engine/Plugins/Animation/MLDeformer/NearestNeighborModel/Source/NearestNeighborModel/Public/NearestNeighborModel.h
+++ b/Engine/Plugins/Animation/MLDeformer/NearestNeighborModel/Source/NearestNeighborModel/Public/NearestNeighborModel.h
@@ -69,6 +69,9 @@ struct FSkeletonCachePair
 
 	UPROPERTY(EditAnywhere, BlueprintReadWrite, Category = "Nearest Neighbors")
 	TObjectPtr<UGeometryCache> Cache = nullptr;
+
+	UPROPERTY(EditAnywhere, BlueprintReadWrite, Category = "Nearest Neighbors")
+	TArray<int32> ExcludedFrames;
 };
 
 USTRUCT(BlueprintType, Blueprintable)
@@ -120,12 +123,19 @@ struct FClothPartData
 	UPROPERTY()
 	TArray<float> VertexMean;
 
+	/** PCA coefficients of the nearest neighbors. This is a flattened array of size (NumNeighbors, PCACoeffNum) */
+	UPROPERTY()
+	TArray<float> AssetNeighborCoeffs;
+
 	/** The remaining offsets of the nearest neighbor shapes (after reducing PCA offsets). This is a flattened array of size (NumNeighbors, PCACoeffNum) */
 	UPROPERTY()
-	TArray<float> NeighborOffsets;
+	TArray<float> AssetNeighborOffsets;
+
+	/** Mapping from NeighborCoeffs to AssetNeighborCoeffs */
+	UPROPERTY()
+	TArray<int32> AssetNeighborIndexMap;
 #endif
 
-	/** PCA coefficients of the nearest neighbors. This is a flattened array of size (NumNeighbors, PCACoeffNum) */
 	UPROPERTY()
 	TArray<float> NeighborCoeffs;
 };
@@ -197,7 +207,6 @@ private:
 	UFUNCTION(BlueprintPure, Category = "Nearest Neighbor Model")
 	int32 GetPCACoeffNum(int32 PartId) const { return ClothPartData[PartId].PCACoeffNum; }
 
-	UFUNCTION(BlueprintPure, Category = "Nearest Neighbor Model")
 	int32 GetNumNeighbors(int32 PartId) const { return ClothPartData[PartId].NumNeighbors; }
 
 	int32 GetTotalNumPCACoeffs() const;
@@ -206,14 +215,6 @@ private:
 	float GetDecayFactor() const { return DecayFactor; };
 	float GetNearestNeighborOffsetWeight() const { return NearestNeighborOffsetWeight; }
 
-	UFUNCTION(BlueprintCallable, Category = "Nearest Neighbor Model")
-	void SetNumNeighbors(int32 PartId, int32 InNumNeighbors) { ClothPartData[PartId].NumNeighbors = InNumNeighbors; }
-	
-	UFUNCTION(BlueprintPure, Category = "Nearest Neighbor Model")
-	const TArray<float>& NeighborCoeffs(int32 PartId) const { return ClothPartData[PartId].NeighborCoeffs; }
-	
-	TArray<float>& NeighborCoeffs(int32 PartId) { return ClothPartData[PartId].NeighborCoeffs; }
-
 	UFUNCTION(BlueprintPure, Category = "Nearest Neighbor Model")
 	TArray<float> ClipInputs(const TArray<float>& Input) const;
 
@@ -223,6 +224,8 @@ private:
 
 	void ResetMorphBuffers();
 
+	const TArray<float>& NeighborCoeffs(int32 PartId) const { return ClothPartData[PartId].NeighborCoeffs; }
+
 	virtual int32 GetNumFloatsPerBone() const override { return NearestNeighborNumFloatsPerBone; }
 	virtual int32 GetNumFloatsPerCurve() const override { return NearestNeighborNumFloatsPerCurve; }
 
@@ -238,6 +241,9 @@ private:
 	bool LoadOptimizedNetworkFromFile(const FString& Filename);
 
 #if WITH_EDITOR
+	UFUNCTION(BlueprintPure, Category = "Nearest Neighbor Model")
+	int32 GetAssetNumNeighbors(int32 PartId) const;
+
 	const TArray<uint32>& PartVertexMap(int32 PartId) const { return ClothPartData[PartId].VertexMap; }
 
 	UFUNCTION(BlueprintPure, Category = "Nearest Neighbor Model")
@@ -256,17 +262,20 @@ private:
 	UFUNCTION(BlueprintCallable, Category = "Nearest Neighbor Model")
 	void SetVertexMean(int32 PartId, const TArray<float>& VertexMean) { ClothPartData[PartId].VertexMean = VertexMean; }
 
+	void SetNumNeighbors(int32 PartId, int32 InNumNeighbors) { ClothPartData[PartId].NumNeighbors = InNumNeighbors; }
+
 	UFUNCTION(BlueprintCallable, Category = "Nearest Neighbor Model")
-	void SetNeighborCoeffs(int32 PartId, const TArray<float>& NeighborCoeffs) { ClothPartData[PartId].NeighborCoeffs = NeighborCoeffs; }
+	void SetAssetNeighborCoeffs(int32 PartId, const TArray<float>& AssetNeighborCoeffs) { ClothPartData[PartId].AssetNeighborCoeffs = AssetNeighborCoeffs; }
 
+	void SetNeighborCoeffs(int32 PartId, const TArray<float>& NeighborCoeffs) { ClothPartData[PartId].NeighborCoeffs = NeighborCoeffs; }
 
 	UFUNCTION(BlueprintPure, Category = "Nearest Neighbor Model")
-	const TArray<float>& NeighborOffsets(int32 PartId) const { return ClothPartData[PartId].NeighborOffsets; }
-
-	TArray<float>& NeighborOffsets(int32 PartId) { return ClothPartData[PartId].NeighborOffsets; }
+	const TArray<float>& AssetNeighborCoeffs(int32 PartId) const { return ClothPartData[PartId].AssetNeighborCoeffs; }
+	
+	const TArray<float>& AssetNeighborOffsets(int32 PartId) const { return ClothPartData[PartId].AssetNeighborOffsets; }
 
 	UFUNCTION(BlueprintCallable, Category = "Nearest Neighbor Model")
-	void SetNeighborOffsets(int32 PartId, const TArray<float>& NeighborOffsets) { ClothPartData[PartId].NeighborOffsets = NeighborOffsets; }
+	void SetAssetNeighborOffsets(int32 PartId, const TArray<float>& AssetNeighborOffsets) { ClothPartData[PartId].AssetNeighborOffsets = AssetNeighborOffsets; }
 #endif
 
 #if WITH_EDITORONLY_DATA
diff --git a/Engine/Plugins/Animation/MLDeformer/NearestNeighborModel/Source/NearestNeighborModel/Public/NearestNeighborModelVizSettings.h b/Engine/Plugins/Animation/MLDeformer/NearestNeighborModel/Source/NearestNeighborModel/Public/NearestNeighborModelVizSettings.h
--- a/Engine/Plugins/Animation/MLDeformer/NearestNeighborModel/Source/NearestNeighborModel/Public/NearestNeighborModelVizSettings.h
+++ b/Engine/Plugins/Animation/MLDeformer/NearestNeighborModel/Source/NearestNeighborModel/Public/NearestNeighborModelVizSettings.h
@@ -24,6 +24,7 @@ class NEARESTNEIGHBORMODEL_API UNearestNeighborModelVizSettings
 public:
 	static FName GetNearestNeighborActorsOffsetPropertyName() { return GET_MEMBER_NAME_CHECKED(UNearestNeighborModelVizSettings, NearestNeighborActorsOffset); }
 	static FName GetNearestNeighborIdsPropertyName() { return GET_MEMBER_NAME_CHECKED(UNearestNeighborModelVizSettings, NearestNeighborIds); }
+	static FName GetNeighborStatsPartIdPropertyName() { return GET_MEMBER_NAME_CHECKED(UNearestNeighborModelVizSettings, NeighborStatsPartId); }
 	
 	friend class UE::NearestNeighborModel::FNearestNeighborEditorModel;
 
@@ -34,6 +35,9 @@ protected:
 	UPROPERTY(VisibleAnywhere, Category = "Live Settings")
 	TArray<uint32> NearestNeighborIds;
 
+	UPROPERTY(EditAnywhere, Category = "Live Settings", meta = (DisplayName = "Part Id"))
+	int32 NeighborStatsPartId = 0;
+
 private:
 	void SetNearestNeighborActorsOffset(float InOffset) { NearestNeighborActorsOffset = InOffset; }
 	float GetNearestNeighborActorsOffset() const { return NearestNeighborActorsOffset; }
diff --git a/Engine/Plugins/Animation/MLDeformer/NearestNeighborModel/Source/NearestNeighborModelEditor/Private/NearestNeighborEditorModel.cpp b/Engine/Plugins/Animation/MLDeformer/NearestNeighborModel/Source/NearestNeighborModelEditor/Private/NearestNeighborEditorModel.cpp
--- a/Engine/Plugins/Animation/MLDeformer/NearestNeighborModel/Source/NearestNeighborModelEditor/Private/NearestNeighborEditorModel.cpp
+++ b/Engine/Plugins/Animation/MLDeformer/NearestNeighborModel/Source/NearestNeighborModelEditor/Private/NearestNeighborEditorModel.cpp
@@ -181,8 +181,19 @@ namespace UE::NearestNeighborModel
 				}
 			}
 			const UNearestNeighborModelInstance* ModelInstance = static_cast<UNearestNeighborModelInstance*>(GetTestMLDeformerModelInstance());
-			if (ModelInstance)
+			UNearestNeighborModel* NearestNeighborModel = GetNearestNeighborModel();
+			if (ModelInstance && NearestNeighborModel)
 			{
+				const TArray<uint32>& NeighborIds = ModelInstance->GetNearestNeighborIds();
+				TArray<uint32> AssetNeighborIds; 
+				const int32 NumParts = FMath::Min(NeighborIds.Num(), NearestNeighborModel->GetNumParts());
+				AssetNeighborIds.SetNum(NumParts);
+				for(int32 Index = 0; Index < NumParts; Index++)
+				{
+					const int32 Id = NeighborIds[Index];
+					const TArray<int32>& IndexMap = NearestNeighborModel->ClothPartData[Index].AssetNeighborIndexMap;
+					AssetNeighborIds[Index] = IndexMap.IsValidIndex(Id) ? IndexMap[Id] : INDEX_NONE;
+				}
 				NNViz->SetNearestNeighborIds(ModelInstance->GetNearestNeighborIds());
 			}
 		}
@@ -544,7 +555,7 @@ namespace UE::NearestNeighborModel
 	{
 		TArray<T> Result;
 		Result.SetNum(End);
-		for (uint32 i = 0; i < End; i++)
+		for (uint32 i = 0; (T)i < End; i++)
 		{
 			Result[i] = i;
 		}
@@ -695,7 +706,25 @@ namespace UE::NearestNeighborModel
 		return MakeTuple(Anim, ReturnCode);
 	}
 
-	void FNearestNeighborEditorModel::AddFloatArrayToDeltaArray(const TArray<float>& FloatArr, const TArray<uint32>& VertexMap, TArray<FVector3f>& DeltaArr, int32 DeltaArrayOffset, float ScaleFactor)
+	namespace Private
+	{
+		TArray<int32> GetIncludedFrame(TConstArrayView<int32> ExcludedFrames, int32 NumFrames)
+		{
+			TArray<int32> IncludedFrames;
+			TArray<int32> FrameMap;
+			IncludedFrames.Reserve(NumFrames);
+			for (int32 Frame = 0; Frame < NumFrames; Frame++)
+			{
+				if (!ExcludedFrames.Contains(Frame))
+				{
+					IncludedFrames.Add(Frame);
+				}
+			}
+			return IncludedFrames;
+		}
+	};
+
+	void FNearestNeighborEditorModel::AddFloatArrayToDeltaArray(const TArray<float>& FloatArr, const TArray<uint32>& VertexMap, TArray<FVector3f>& DeltaArr, int32 DeltaArrayOffset, TOptional<TArray<int32>> OptionalIncludedFrames)
 	{
 		const int32 NumBaseMeshVerts = Model->GetNumBaseMeshVerts();
 		const UNearestNeighborModel* NearestNeighborModel = GetNearestNeighborModel();
@@ -705,19 +734,22 @@ namespace UE::NearestNeighborModel
 			return;
 		}
 		const int32 NumShapes = FloatArr.Num() / (PartNumVerts * 3);
-		if (DeltaArrayOffset < 0)
+		if (DeltaArrayOffset == INDEX_NONE)
 		{
 			DeltaArrayOffset = DeltaArr.Num();
 		}
-		DeltaArr.SetNumZeroed(FMath::Max(DeltaArrayOffset + NumShapes * NumBaseMeshVerts, DeltaArr.Num()), false);
+		const TArray<int32> IncludedFrames = OptionalIncludedFrames.IsSet() ? OptionalIncludedFrames.GetValue() : Range<int32>(NumShapes);
+		const int32 NumIncludedFrames = IncludedFrames.Num();
+		DeltaArr.SetNumZeroed(FMath::Max(DeltaArrayOffset + NumIncludedFrames * NumBaseMeshVerts, DeltaArr.Num()), false);
 
-		for(int32 ShapeId = 0; ShapeId < NumShapes; ShapeId++)
+		for(int32 Frame = 0; Frame < IncludedFrames.Num(); ++Frame)
 		{
+			const int32 IncludedFrame = IncludedFrames[Frame];
 			for(int32 VertexId = 0; VertexId < PartNumVerts; VertexId++)
 			{
-				const int32 DeltaId = ShapeId * NumBaseMeshVerts + VertexMap[VertexId];
-				const int32 FloatId = (ShapeId * PartNumVerts + VertexId) * 3;
-				DeltaArr[DeltaArrayOffset + DeltaId] = FVector3f(FloatArr[FloatId], FloatArr[FloatId + 1], FloatArr[FloatId + 2]) * ScaleFactor;
+				const int32 DeltaId = Frame * NumBaseMeshVerts + VertexMap[VertexId];
+				const int32 FloatId = (IncludedFrame * PartNumVerts + VertexId) * 3;
+				DeltaArr[DeltaArrayOffset + DeltaId] = FVector3f(FloatArr[FloatId], FloatArr[FloatId + 1], FloatArr[FloatId + 2]);
 			}
 		}
 	}
@@ -824,8 +856,22 @@ namespace UE::NearestNeighborModel
 
 		for (int32 PartId = 0; PartId < NumParts; PartId++)
 		{
-			const TArray<uint32>& VertexMap = NearestNeighborModel->PartVertexMap(PartId);
-			AddFloatArrayToDeltaArray(NearestNeighborModel->ClothPartData[PartId].NeighborOffsets, VertexMap, Deltas);
+			FClothPartData& Data = NearestNeighborModel->ClothPartData[PartId];
+			const TArray<uint32>& VertexMap = Data.VertexMap;
+			const TArray<float>& AssetNeighborOffsets = Data.AssetNeighborOffsets;
+			constexpr int32 DeltaOffset = INDEX_NONE;
+			const int32 NumAssetNeighbors = AssetNeighborOffsets.Num() / (VertexMap.Num() * 3);
+			check(AssetNeighborOffsets.Num() % (VertexMap.Num() * 3) == 0);
+			if (NumAssetNeighbors == 0)
+			{
+				continue;
+			}
+			check(NearestNeighborModel->NearestNeighborData.IsValidIndex(PartId));
+			const TArray<int32> ExcludedFrames = NearestNeighborModel->NearestNeighborData[PartId].ExcludedFrames;
+			const TArray<int32> IncludedFrames = UE::NearestNeighborModel::Private::GetIncludedFrame(ExcludedFrames, NumAssetNeighbors);
+			NearestNeighborModel->SetNumNeighbors(PartId, IncludedFrames.Num());
+			Data.AssetNeighborIndexMap = IncludedFrames;
+			AddFloatArrayToDeltaArray(AssetNeighborOffsets, VertexMap, Deltas, DeltaOffset, IncludedFrames);
 		}
 
 		if (Deltas.Num() == 0)
@@ -938,6 +984,20 @@ namespace UE::NearestNeighborModel
 	void FNearestNeighborEditorModel::GetNeighborStats()
 	{
 		UNearestNeighborTrainingModel *TrainingModel = InitTrainingModel<UNearestNeighborTrainingModel>(this);
+		if (!TrainingModel)
+		{
+			return;
+		}
+		const UNearestNeighborModel *NearestNeighborModel = GetNearestNeighborModel();
+		if (!NearestNeighborModel)
+		{
+			return;
+		}
+		if (!NearestNeighborModel->GetOptimizedNetwork())
+		{
+			UE_LOG(LogNearestNeighborModel, Error, TEXT("Network is not loaded. Network is needed to compute stats. Please train model first."));
+			return;
+		}
 		TrainingModel->GetNeighborStats(0);
 	}
 }	// namespace UE::NearestNeighborModel
diff --git a/Engine/Plugins/Animation/MLDeformer/NearestNeighborModel/Source/NearestNeighborModelEditor/Private/NearestNeighborEditorModelActor.cpp b/Engine/Plugins/Animation/MLDeformer/NearestNeighborModel/Source/NearestNeighborModelEditor/Private/NearestNeighborEditorModelActor.cpp
--- a/Engine/Plugins/Animation/MLDeformer/NearestNeighborModel/Source/NearestNeighborModelEditor/Private/NearestNeighborEditorModelActor.cpp
+++ b/Engine/Plugins/Animation/MLDeformer/NearestNeighborModel/Source/NearestNeighborModelEditor/Private/NearestNeighborEditorModelActor.cpp
@@ -6,8 +6,6 @@
 
 namespace UE::NearestNeighborModel
 {
-	using namespace UE::MLDeformer;
-
 	FNearestNeighborEditorModelActor::FNearestNeighborEditorModelActor(const FConstructSettings& Settings)
 		: FMLDeformerGeomCacheActor(Settings)
 	{
diff --git a/Engine/Plugins/Animation/MLDeformer/NearestNeighborModel/Source/NearestNeighborModelEditor/Private/NearestNeighborGeomCacheSampler.cpp b/Engine/Plugins/Animation/MLDeformer/NearestNeighborModel/Source/NearestNeighborModelEditor/Private/NearestNeighborGeomCacheSampler.cpp
--- a/Engine/Plugins/Animation/MLDeformer/NearestNeighborModel/Source/NearestNeighborModelEditor/Private/NearestNeighborGeomCacheSampler.cpp
+++ b/Engine/Plugins/Animation/MLDeformer/NearestNeighborModel/Source/NearestNeighborModelEditor/Private/NearestNeighborGeomCacheSampler.cpp
@@ -50,7 +50,7 @@ namespace UE::NearestNeighborModel
 			const TArray<FSkelMeshImportedMeshInfo>& SkelMeshInfos = LODModel.ImportedMeshInfos;
 			for (int32 MeshMappingIndex = 0; MeshMappingIndex < MeshMappings.Num(); ++MeshMappingIndex)
 			{
-				const UE::MLDeformer::FMLDeformerGeomCacheMeshMapping& MeshMapping = MeshMappings[MeshMappingIndex];
+				const ::UE::MLDeformer::FMLDeformerGeomCacheMeshMapping& MeshMapping = MeshMappings[MeshMappingIndex];
 				const FSkelMeshImportedMeshInfo& MeshInfo = SkelMeshInfos[MeshMapping.MeshIndex];
 				UGeometryCacheTrack* Track = GeometryCache->Tracks[MeshMapping.TrackIndex];
 
@@ -153,7 +153,7 @@ namespace UE::NearestNeighborModel
 				{
 					Sign = Inner(R0, R) < 0 ? -1 : 1;
 				}
-				const float	Weight = Sign * static_cast<float>(WeightByte) * UE::AnimationCore::InvMaxRawBoneWeightFloat;
+				const float	Weight = Sign * static_cast<float>(WeightByte) * ::UE::AnimationCore::InvMaxRawBoneWeightFloat;
 				QuatSum += R * Weight;
 			}
 		}
diff --git a/Engine/Plugins/Animation/MLDeformer/NearestNeighborModel/Source/NearestNeighborModelEditor/Private/NearestNeighborModelDetails.cpp b/Engine/Plugins/Animation/MLDeformer/NearestNeighborModel/Source/NearestNeighborModelEditor/Private/NearestNeighborModelDetails.cpp
--- a/Engine/Plugins/Animation/MLDeformer/NearestNeighborModel/Source/NearestNeighborModelEditor/Private/NearestNeighborModelDetails.cpp
+++ b/Engine/Plugins/Animation/MLDeformer/NearestNeighborModel/Source/NearestNeighborModelEditor/Private/NearestNeighborModelDetails.cpp
@@ -25,7 +25,7 @@
 
 namespace UE::NearestNeighborModel
 {
-	using namespace UE::MLDeformer;
+	using namespace ::UE::MLDeformer;
 
 	TSharedRef<IDetailCustomization> FNearestNeighborModelDetails::MakeInstance()
 	{
diff --git a/Engine/Plugins/Animation/MLDeformer/NearestNeighborModel/Source/NearestNeighborModelEditor/Private/NearestNeighborModelEditorModule.cpp b/Engine/Plugins/Animation/MLDeformer/NearestNeighborModel/Source/NearestNeighborModelEditor/Private/NearestNeighborModelEditorModule.cpp
--- a/Engine/Plugins/Animation/MLDeformer/NearestNeighborModel/Source/NearestNeighborModelEditor/Private/NearestNeighborModelEditorModule.cpp
+++ b/Engine/Plugins/Animation/MLDeformer/NearestNeighborModel/Source/NearestNeighborModelEditor/Private/NearestNeighborModelEditorModule.cpp
@@ -26,8 +26,6 @@ IMPLEMENT_MODULE(UE::NearestNeighborModel::FNearestNeighborModelEditorModule, Ne
 
 namespace UE::NearestNeighborModel
 {
-	using namespace UE::MLDeformer;
-
 	void FNearestNeighborModelEditorModule::StartupModule()
 	{
 		// Register object detail customizations.
diff --git a/Engine/Plugins/Animation/MLDeformer/NearestNeighborModel/Source/NearestNeighborModelEditor/Private/NearestNeighborModelVizSettingsDetails.cpp b/Engine/Plugins/Animation/MLDeformer/NearestNeighborModel/Source/NearestNeighborModelEditor/Private/NearestNeighborModelVizSettingsDetails.cpp
--- a/Engine/Plugins/Animation/MLDeformer/NearestNeighborModel/Source/NearestNeighborModelEditor/Private/NearestNeighborModelVizSettingsDetails.cpp
+++ b/Engine/Plugins/Animation/MLDeformer/NearestNeighborModel/Source/NearestNeighborModelEditor/Private/NearestNeighborModelVizSettingsDetails.cpp
@@ -23,6 +23,7 @@ namespace UE::NearestNeighborModel
 		NNGroup.AddPropertyRow(DetailLayoutBuilder->GetProperty(UNearestNeighborModelVizSettings::GetNearestNeighborIdsPropertyName(), UNearestNeighborModelVizSettings::StaticClass()));
 
 		IDetailGroup& ToolsGroup = LiveSettingsCategory->AddGroup("Tools", LOCTEXT("ToolsLabel", "Tools"), false, true);
+		ToolsGroup.AddPropertyRow(DetailLayoutBuilder->GetProperty(UNearestNeighborModelVizSettings::GetNeighborStatsPartIdPropertyName(), UNearestNeighborModelVizSettings::StaticClass()));
 		ToolsGroup.AddWidgetRow()
 		.WholeRowContent()
 		[
diff --git a/Engine/Plugins/Animation/MLDeformer/NearestNeighborModel/Source/NearestNeighborModelEditor/Private/NearestNeighborTrainingModel.cpp b/Engine/Plugins/Animation/MLDeformer/NearestNeighborModel/Source/NearestNeighborModelEditor/Private/NearestNeighborTrainingModel.cpp
--- a/Engine/Plugins/Animation/MLDeformer/NearestNeighborModel/Source/NearestNeighborModelEditor/Private/NearestNeighborTrainingModel.cpp
+++ b/Engine/Plugins/Animation/MLDeformer/NearestNeighborModel/Source/NearestNeighborModelEditor/Private/NearestNeighborTrainingModel.cpp
@@ -7,11 +7,10 @@
 #include "Animation/AnimSequence.h"
 
 #define LOCTEXT_NAMESPACE "NearestNeighborTrainingModel"
-using namespace UE::NearestNeighborModel;
 
 UNearestNeighborTrainingModel::~UNearestNeighborTrainingModel() = default;
 
-void UNearestNeighborTrainingModel::Init(UE::MLDeformer::FMLDeformerEditorModel* InEditorModel)
+void UNearestNeighborTrainingModel::Init(::UE::MLDeformer::FMLDeformerEditorModel* InEditorModel)
 {
 	UMLDeformerTrainingModel::Init(InEditorModel);
 	check(InEditorModel != nullptr);
@@ -121,8 +120,8 @@ int32 UNearestNeighborTrainingModel::GetKmeansNumClusters() const
 
 const TArray<float> UNearestNeighborTrainingModel::GetUnskinnedVertexPositions() const
 {
-	FMLDeformerSampler* BaseSampler = EditorModel->GetNumTrainingInputAnims() > 0 ? EditorModel->GetSamplerForTrainingAnim(0) : nullptr;
-	FNearestNeighborGeomCacheSampler* Sampler = static_cast<FNearestNeighborGeomCacheSampler*>(BaseSampler);
+	UE::MLDeformer::FMLDeformerSampler* BaseSampler = EditorModel->GetNumTrainingInputAnims() > 0 ? EditorModel->GetSamplerForTrainingAnim(0) : nullptr;
+	UE::NearestNeighborModel::FNearestNeighborGeomCacheSampler* Sampler = static_cast<UE::NearestNeighborModel::FNearestNeighborGeomCacheSampler*>(BaseSampler);
 	if (Sampler == nullptr)
 	{
 		return TArray<float>();
@@ -142,13 +141,12 @@ const TArray<float> UNearestNeighborTrainingModel::GetUnskinnedVertexPositions()
 
 const TArray<int32> UNearestNeighborTrainingModel::GetMeshIndexBuffer() const
 {
-	FMLDeformerSampler* BaseSampler = EditorModel->GetNumTrainingInputAnims() > 0 ? EditorModel->GetSamplerForTrainingAnim(0) : nullptr;
-	FNearestNeighborGeomCacheSampler* Sampler = static_cast<FNearestNeighborGeomCacheSampler*>(BaseSampler);
+	UE::MLDeformer::FMLDeformerSampler* BaseSampler = EditorModel->GetNumTrainingInputAnims() > 0 ? EditorModel->GetSamplerForTrainingAnim(0) : nullptr;
+	UE::NearestNeighborModel::FNearestNeighborGeomCacheSampler* Sampler = static_cast<UE::NearestNeighborModel::FNearestNeighborGeomCacheSampler*>(BaseSampler);
 	if (Sampler == nullptr)
 	{
 		return TArray<int32>();
 	}
-
 	TArray<uint32> UIndexBuffer = Sampler->GetMeshIndexBuffer();
 	return TArray<int32>((int32*)UIndexBuffer.GetData(), UIndexBuffer.Num());
 }
@@ -169,7 +167,7 @@ void UNearestNeighborTrainingModel::DestroyModelInstance(UNearestNeighborModelIn
 
 void UNearestNeighborTrainingModel::SetNewAnimSampler()
 {
-	AnimSampler = MakeUnique<FNearestNeighborGeomCacheSampler>();
+	AnimSampler = MakeUnique<UE::NearestNeighborModel::FNearestNeighborGeomCacheSampler>();
 	AnimSampler->Init(EditorModel, 0);
 }
 
@@ -179,7 +177,7 @@ const UAnimSequence* UNearestNeighborTrainingModel::GetTestAnim() const
 	{
 		return nullptr;
 	}
-	const FNearestNeighborEditorModel* NNEditorModel = static_cast<FNearestNeighborEditorModel*>(EditorModel);
+	const UE::NearestNeighborModel::FNearestNeighborEditorModel* NNEditorModel = static_cast<UE::NearestNeighborModel::FNearestNeighborEditorModel*>(EditorModel);
 	const UNearestNeighborModelVizSettings* Viz = NNEditorModel->GetNearestNeighborModelVizSettings();
 	if (!Viz)
 	{
diff --git a/Engine/Plugins/Animation/MLDeformer/NearestNeighborModel/Source/NearestNeighborModelEditor/Public/NearestNeighborEditorModel.h b/Engine/Plugins/Animation/MLDeformer/NearestNeighborModel/Source/NearestNeighborModelEditor/Public/NearestNeighborEditorModel.h
--- a/Engine/Plugins/Animation/MLDeformer/NearestNeighborModel/Source/NearestNeighborModelEditor/Public/NearestNeighborEditorModel.h
+++ b/Engine/Plugins/Animation/MLDeformer/NearestNeighborModel/Source/NearestNeighborModelEditor/Public/NearestNeighborEditorModel.h
@@ -81,7 +81,7 @@ namespace UE::NearestNeighborModel
 
 		uint8 InitMorphTargets();
 		void RefreshMorphTargets();
-		void AddFloatArrayToDeltaArray(const TArray<float>& FloatArr, const TArray<uint32>& VertexMap, TArray<FVector3f>& DeltaArr, int32 DeltaArrayOffset = -1, float ScaleFactor = 1);
+		void AddFloatArrayToDeltaArray(const TArray<float>& FloatArr, const TArray<uint32>& VertexMap, TArray<FVector3f>& DeltaArr, int32 DeltaArrayOffset = INDEX_NONE, TOptional<TArray<int32>> OptionalIncludedFrames = TOptional<TArray<int32>>());
 
 		int32 GetNumParts();
 
diff --git a/Engine/Plugins/Animation/MLDeformer/NearestNeighborModel/Source/NearestNeighborModelEditor/Public/NearestNeighborEditorModelActor.h b/Engine/Plugins/Animation/MLDeformer/NearestNeighborModel/Source/NearestNeighborModelEditor/Public/NearestNeighborEditorModelActor.h
--- a/Engine/Plugins/Animation/MLDeformer/NearestNeighborModel/Source/NearestNeighborModelEditor/Public/NearestNeighborEditorModelActor.h
+++ b/Engine/Plugins/Animation/MLDeformer/NearestNeighborModel/Source/NearestNeighborModelEditor/Public/NearestNeighborEditorModelActor.h
@@ -12,15 +12,13 @@ namespace UE::NearestNeighborModel
 {
 	class FNearestNeighborEditorModel;
 
-	using namespace UE::MLDeformer;
-
 	enum : int32
 	{
 		ActorID_NearestNeighborActors = 6
 	};
 
 	class NEARESTNEIGHBORMODELEDITOR_API FNearestNeighborEditorModelActor
-		: public FMLDeformerGeomCacheActor
+		: public ::UE::MLDeformer::FMLDeformerGeomCacheActor
 	{
 	public:
 		FNearestNeighborEditorModelActor(const FConstructSettings& Settings);
diff --git a/Engine/Plugins/Animation/MLDeformer/NearestNeighborModel/Source/NearestNeighborModelEditor/Public/NearestNeighborModelDetails.h b/Engine/Plugins/Animation/MLDeformer/NearestNeighborModel/Source/NearestNeighborModelEditor/Public/NearestNeighborModelDetails.h
--- a/Engine/Plugins/Animation/MLDeformer/NearestNeighborModel/Source/NearestNeighborModelEditor/Public/NearestNeighborModelDetails.h
+++ b/Engine/Plugins/Animation/MLDeformer/NearestNeighborModel/Source/NearestNeighborModelEditor/Public/NearestNeighborModelDetails.h
@@ -15,7 +15,7 @@ namespace UE::NearestNeighborModel
 	class FNearestNeighborEditorModel;
 
 	class NEARESTNEIGHBORMODELEDITOR_API FNearestNeighborModelDetails
-		: public UE::MLDeformer::FMLDeformerMorphModelDetails
+		: public ::UE::MLDeformer::FMLDeformerMorphModelDetails
 	{
 	public:
 		/** Makes a new instance of this detail layout class for a specific detail view requesting it. */
diff --git a/Engine/Plugins/Animation/MLDeformer/NearestNeighborModel/Source/NearestNeighborModelEditor/Public/NearestNeighborModelVizSettingsDetails.h b/Engine/Plugins/Animation/MLDeformer/NearestNeighborModel/Source/NearestNeighborModelEditor/Public/NearestNeighborModelVizSettingsDetails.h
--- a/Engine/Plugins/Animation/MLDeformer/NearestNeighborModel/Source/NearestNeighborModelEditor/Public/NearestNeighborModelVizSettingsDetails.h
+++ b/Engine/Plugins/Animation/MLDeformer/NearestNeighborModel/Source/NearestNeighborModelEditor/Public/NearestNeighborModelVizSettingsDetails.h
@@ -10,7 +10,7 @@ namespace UE::NearestNeighborModel
 	class FNearestNeighborEditorModel;
 
 	class NEARESTNEIGHBORMODELEDITOR_API FNearestNeighborModelVizSettingsDetails
-		: public UE::MLDeformer::FMLDeformerMorphModelVizSettingsDetails
+		: public ::UE::MLDeformer::FMLDeformerMorphModelVizSettingsDetails
 	{
 	public:
 		/** Makes a new instance of this detail layout class for a specific detail view requesting it. */
