NearestNeighborModel: add RBF to alleviate "living cloth" problem

diff --git a/Engine/Plugins/Animation/MLDeformer/NearestNeighborModel/Source/NearestNeighborModel/Private/NearestNeighborModelInstance.cpp b/Engine/Plugins/Animation/MLDeformer/NearestNeighborModel/Source/NearestNeighborModel/Private/NearestNeighborModelInstance.cpp
--- a/Engine/Plugins/Animation/MLDeformer/NearestNeighborModel/Source/NearestNeighborModel/Private/NearestNeighborModelInstance.cpp
+++ b/Engine/Plugins/Animation/MLDeformer/NearestNeighborModel/Source/NearestNeighborModel/Private/NearestNeighborModelInstance.cpp
@@ -159,9 +159,66 @@ TArray<float> ExtractArray(const TArray<float>& InArr, int32 Start, int32 End)
 
 constexpr float FastExpMinusOne(float X)
 {
-	return X + X * X / 2 + X * X * X / 6;
+	return X + X * X / 2.0f + X * X * X / 6.0f;
 }
 
+namespace UE::NearestNeighborModel::Private
+{
+	float ComputeDistanceSquared(TConstArrayView<float> V1, TConstArrayView<float> V2)
+	{
+		check(V1.Num() == V2.Num());
+		float DistanceSquared = 0;
+		for (int32 Index = 0; Index < V1.Num(); ++Index)
+		{
+			const float D = V1[Index] - V2[Index];
+			DistanceSquared += D * D;
+		}
+		return DistanceSquared;
+	}
+
+	void UpdateRBFWeights(TArrayView<float> OutMorphWeights, int32& OutNearestNeighborId, TArrayView<float> OutPreviousWeights, TConstArrayView<float> CoeffsView, const TConstArrayView<float> NeighborCoeffsView, float Sigma, float OffsetWeight, float DecayCoeff)
+	{
+		const int32 NumPCACoeffs = CoeffsView.Num();
+		check(OutMorphWeights.Num() == NumPCACoeffs);
+		const int32 NumNeighbors = NeighborCoeffsView.Num() / NumPCACoeffs;
+		check(NeighborCoeffsView.Num() == NumPCACoeffs * NumNeighbors);
+		TArray<float> Weights;
+		Weights.SetNumUninitialized(NumNeighbors);
+		for (int32 Index = 0; Index < NumNeighbors; ++Index)
+		{
+			TConstArrayView<float> NeighborCoeffs(NeighborCoeffsView.GetData() + Index * NumPCACoeffs, NumPCACoeffs);
+			Weights[Index] = ComputeDistanceSquared(CoeffsView, NeighborCoeffs);
+		}
+		const float MinD2 = FMath::Min(Weights, &OutNearestNeighborId);
+		float SumWeights = 0.0f;
+		const float Sigma2 = FMath::Max(Sigma * Sigma, 1e-6f);
+		for (int32 Index = 0; Index < NumNeighbors; ++Index)
+		{
+			float Weight = (Weights[Index] - MinD2) / Sigma2;
+			constexpr float CutOff = 3.0f;
+			if (Weight < CutOff)
+			{
+				Weight = FMath::Exp(-Weight);
+				Weights[Index] = Weight;
+				SumWeights += Weight;
+			}
+			else
+			{
+				Weights[Index] = 0.0f;
+			}
+		}
+		SumWeights = FMath::Max(SumWeights, 1e-6f);
+		for (int32 Index = 0; Index < NumNeighbors; ++Index)
+		{
+			const float W = Weights[Index] / SumWeights * OffsetWeight;
+			const float PreviousW = OutPreviousWeights[Index];
+			const float NewW = (1.0f - DecayCoeff) * W + DecayCoeff * PreviousW;
+			OutMorphWeights[Index] = NewW;
+			OutPreviousWeights[Index] = NewW;
+		}
+	}
+};
+
 void UNearestNeighborModelInstance::RunNearestNeighborModel(float DeltaTime, float ModelWeight)
 {
 	UNearestNeighborModel* NearestNeighborModel = Cast<UNearestNeighborModel>(Model);
@@ -214,22 +271,46 @@ void UNearestNeighborModelInstance::RunNearestNeighborModel(float DeltaTime, flo
 			int32 NeighborOffset = NumNetworkWeights + 1;
 			for (int32 PartId = 0; PartId < NearestNeighborModel->GetNumParts(); PartId++)
 			{
-				const int32 NearestNeighborId = FindNearestNeighbor(OutputTensorData, PartId);
-	#if WITH_EDITORONLY_DATA
-				NearestNeighborIds[PartId] = NearestNeighborId;
-	#endif
-
-				const int32 NumNeighbors = NearestNeighborModel->GetNumNeighbors(PartId);
-				for (int32 NeighborId = 0; NeighborId < NumNeighbors; NeighborId++)
+				if (NearestNeighborModel->bUseRBF)
+				{
+					const int32 CoeffStart = NearestNeighborModel->GetPCACoeffStart(PartId);
+					const int32 NumPCACoeffs = NearestNeighborModel->GetPCACoeffNum(PartId);
+					TConstArrayView<float> CoeffsView(OutputTensorData + CoeffStart, NumPCACoeffs);
+					check(CoeffStart + NumPCACoeffs <= NumNetworkWeights);
+
+					const TArray<float>& NeighborCoeffs = NearestNeighborModel->NeighborCoeffs(PartId);
+					const int32 NumNeighbors = NearestNeighborModel->GetNumNeighbors(PartId);
+					check(NeighborCoeffs.Num() == NumPCACoeffs * NumNeighbors);
+
+					TArrayView<float> MorphWeightsView(WeightData->Weights.GetData() + NeighborOffset, NumPCACoeffs);
+					TArrayView<float> PreviousWeightsView(PreviousWeights.GetData() + NeighborOffset, NumPCACoeffs);
+					const float OffsetWeight = ModelWeight * NearestNeighborModel->GetNearestNeighborOffsetWeight();
+					int32 NearestNeighborId = INDEX_NONE;
+					Private::UpdateRBFWeights(MorphWeightsView, NearestNeighborId, PreviousWeightsView, CoeffsView, NeighborCoeffs, NearestNeighborModel->RBFSigma, OffsetWeight, DecayCoeff);
+					#if WITH_EDITORONLY_DATA
+						NearestNeighborIds[PartId] = NearestNeighborId;
+					#endif
+					NeighborOffset += NumNeighbors;
+				}
+				else
 				{
-					const float W = NeighborId == NearestNeighborId ? ModelWeight * NearestNeighborModel->GetNearestNeighborOffsetWeight() : 0;
-					const int32 Index = NeighborOffset + NeighborId;
-					if (Index < NumMorphTargets)
+					const int32 NearestNeighborId = FindNearestNeighbor(OutputTensorData, PartId);
+						#if WITH_EDITORONLY_DATA
+					NearestNeighborIds[PartId] = NearestNeighborId;
+						#endif
+					
+					const int32 NumNeighbors = NearestNeighborModel->GetNumNeighbors(PartId);
+					for (int32 NeighborId = 0; NeighborId < NumNeighbors; NeighborId++)
 					{
-						UpdateWeightWithDecay(WeightData->Weights, Index, W, DecayCoeff);
+						const float W = NeighborId == NearestNeighborId ? ModelWeight * NearestNeighborModel->GetNearestNeighborOffsetWeight() : 0;
+						const int32 Index = NeighborOffset + NeighborId;
+						if (Index < NumMorphTargets)
+						{
+							UpdateWeightWithDecay(WeightData->Weights, Index, W, DecayCoeff);
+						}
 					}
+					NeighborOffset += NumNeighbors;
 				}
-				NeighborOffset += NumNeighbors;
 			}
 		}
 	}
diff --git a/Engine/Plugins/Animation/MLDeformer/NearestNeighborModel/Source/NearestNeighborModel/Public/NearestNeighborModel.h b/Engine/Plugins/Animation/MLDeformer/NearestNeighborModel/Source/NearestNeighborModel/Public/NearestNeighborModel.h
--- a/Engine/Plugins/Animation/MLDeformer/NearestNeighborModel/Source/NearestNeighborModel/Public/NearestNeighborModel.h
+++ b/Engine/Plugins/Animation/MLDeformer/NearestNeighborModel/Source/NearestNeighborModel/Public/NearestNeighborModel.h
@@ -459,6 +459,14 @@ protected:
 	UPROPERTY(EditAnywhere, BlueprintReadWrite, Category = "Nearest Neighbors", META = (ClampMin = "0", ClampMax = "1"))
 	float NearestNeighborOffsetWeight = 1.0f;
 
+	/** Use radial basis function (RBF). This will blend multiple nearest neighbors instead of finding a single neighbor. */
+	UPROPERTY(EditAnywhere, BlueprintReadWrite, Category = "Nearest Neighbors")
+	bool bUseRBF = false;
+
+	/** The smaller this value is, the more similar it is to a single nearest neighbor. The bigger this value is, the more smeared out the geometric details are. */
+	UPROPERTY(EditAnywhere, BlueprintReadWrite, Category = "Nearest Neighbors", META = (ClampMin = "0.001"))
+	float RBFSigma = 1.0f;
+
 private:
 	UPROPERTY()
 	TObjectPtr<UNearestNeighborOptimizedNetwork> OptimizedNetwork = nullptr;
