NearestNeighorModel: fix NeighborCoeffs bug

diff --git a/Engine/Plugins/Animation/MLDeformer/NearestNeighborModel/Source/NearestNeighborModel/Private/NearestNeighborModelInstance.cpp b/Engine/Plugins/Animation/MLDeformer/NearestNeighborModel/Source/NearestNeighborModel/Private/NearestNeighborModelInstance.cpp
--- a/Engine/Plugins/Animation/MLDeformer/NearestNeighborModel/Source/NearestNeighborModel/Private/NearestNeighborModelInstance.cpp
+++ b/Engine/Plugins/Animation/MLDeformer/NearestNeighborModel/Source/NearestNeighborModel/Private/NearestNeighborModelInstance.cpp
@@ -280,7 +280,10 @@ void UNearestNeighborModelInstance::RunNearestNeighborModel(float DeltaTime, flo
 
 					const TArray<float>& NeighborCoeffs = NearestNeighborModel->NeighborCoeffs(PartId);
 					const int32 NumNeighbors = NearestNeighborModel->GetNumNeighbors(PartId);
-					check(NeighborCoeffs.Num() == NumPCACoeffs * NumNeighbors);
+					if(!ensureMsgf(NeighborCoeffs.Num() == NumPCACoeffs * NumNeighbors, TEXT("Nearest Neigbbor Data is out-of-date, please hit Update button")))
+					{
+						continue;
+					}
 
 					TArrayView<float> MorphWeightsView(WeightData->Weights.GetData() + NeighborOffset, NumPCACoeffs);
 					TArrayView<float> PreviousWeightsView(PreviousWeights.GetData() + NeighborOffset, NumPCACoeffs);
diff --git a/Engine/Plugins/Animation/MLDeformer/NearestNeighborModel/Source/NearestNeighborModelEditor/Private/NearestNeighborEditorModel.cpp b/Engine/Plugins/Animation/MLDeformer/NearestNeighborModel/Source/NearestNeighborModelEditor/Private/NearestNeighborEditorModel.cpp
--- a/Engine/Plugins/Animation/MLDeformer/NearestNeighborModel/Source/NearestNeighborModelEditor/Private/NearestNeighborEditorModel.cpp
+++ b/Engine/Plugins/Animation/MLDeformer/NearestNeighborModel/Source/NearestNeighborModelEditor/Private/NearestNeighborEditorModel.cpp
@@ -722,6 +722,25 @@ namespace UE::NearestNeighborModel
 			}
 			return IncludedFrames;
 		}
+
+		TArray<float> CompactCoeffs(TConstArrayView<float> SourceCoeffs, int32 NumCoeffs, const TArray<int32>& IncludedFrames)
+		{
+			check(NumCoeffs > 0);
+			const int32 NumSource = SourceCoeffs.Num() / NumCoeffs;
+			check(SourceCoeffs.Num() == NumSource * NumCoeffs);
+			const int32 NumResult = IncludedFrames.Num();
+			TArray<float> ResultCoeffs;
+			ResultCoeffs.SetNum(NumResult * NumCoeffs);
+			for (int32 FrameIndex = 0; FrameIndex < NumResult; ++FrameIndex)
+			{
+				const int32 Frame = IncludedFrames[FrameIndex];
+				for (int32 CoeffIndex = 0; CoeffIndex < NumCoeffs; ++CoeffIndex)
+				{
+					ResultCoeffs[FrameIndex * NumCoeffs + CoeffIndex] = SourceCoeffs[Frame * NumCoeffs + CoeffIndex];
+				}
+			}
+			return ResultCoeffs;
+		}
 	};
 
 	void FNearestNeighborEditorModel::AddFloatArrayToDeltaArray(const TArray<float>& FloatArr, const TArray<uint32>& VertexMap, TArray<FVector3f>& DeltaArr, int32 DeltaArrayOffset, TOptional<TArray<int32>> OptionalIncludedFrames)
@@ -870,6 +889,8 @@ namespace UE::NearestNeighborModel
 			const TArray<int32> ExcludedFrames = NearestNeighborModel->NearestNeighborData[PartId].ExcludedFrames;
 			const TArray<int32> IncludedFrames = UE::NearestNeighborModel::Private::GetIncludedFrame(ExcludedFrames, NumAssetNeighbors);
 			NearestNeighborModel->SetNumNeighbors(PartId, IncludedFrames.Num());
+			const int32 NumPCACoeffs = Data.PCACoeffNum;
+			Data.NeighborCoeffs = Private::CompactCoeffs(Data.AssetNeighborCoeffs, NumPCACoeffs, IncludedFrames);
 			Data.AssetNeighborIndexMap = IncludedFrames;
 			AddFloatArrayToDeltaArray(AssetNeighborOffsets, VertexMap, Deltas, DeltaOffset, IncludedFrames);
 		}
