Force loading of NNERuntimeBasicCpu plugin before construction of runtime.

diff --git a/Engine/Plugins/Animation/MLDeformer/NearestNeighborModel/Source/NearestNeighborModel/Private/NearestNeighborOptimizedNetwork.cpp b/Engine/Plugins/Animation/MLDeformer/NearestNeighborModel/Source/NearestNeighborModel/Private/NearestNeighborOptimizedNetwork.cpp
--- a/Engine/Plugins/Animation/MLDeformer/NearestNeighborModel/Source/NearestNeighborModel/Private/NearestNeighborOptimizedNetwork.cpp
+++ b/Engine/Plugins/Animation/MLDeformer/NearestNeighborModel/Source/NearestNeighborModel/Private/NearestNeighborOptimizedNetwork.cpp
@@ -3,6 +3,7 @@
 #include "NearestNeighborOptimizedNetwork.h"
 
 #include "Misc/FileHelper.h"
+#include "Modules/ModuleManager.h"
 
 #include "NNE.h"
 #include "NNERuntime.h"
@@ -67,6 +68,7 @@ namespace UE::NearestNeighborModel::Private
 PRAGMA_DISABLE_DEPRECATION_WARNINGS
 
 const TCHAR* UNearestNeighborOptimizedNetwork::RuntimeName = TEXT("NNERuntimeBasicCpu");
+const TCHAR* UNearestNeighborOptimizedNetwork::RuntimeModuleName = TEXT("NNERuntimeBasicCpu");
 
 
 void UNearestNeighborOptimizedNetwork::PostLoad()
@@ -96,6 +98,8 @@ void UNearestNeighborOptimizedNetwork::PostLoad()
 
 	// Create in-memory representation of network
 
+	ensureMsgf(FModuleManager::Get().LoadModule(RuntimeModuleName), TEXT("Unable to load runtime module."));
+
 	TWeakInterfacePtr<INNERuntimeCPU> RuntimeCPU = UE::NNE::GetRuntime<INNERuntimeCPU>(RuntimeName);
 
 	if (ensureMsgf(RuntimeCPU.IsValid(), TEXT("Could not find requested NNE Runtime")))
@@ -148,6 +152,8 @@ bool UNearestNeighborOptimizedNetwork::Load(const FString& Filename)
 		// Clear FileData to avoid multiple copies in memory at once
 		FileData.Empty();
 
+		ensureMsgf(FModuleManager::Get().LoadModule(RuntimeModuleName), TEXT("Unable to load runtime module."));
+
 		TWeakInterfacePtr<INNERuntimeCPU> RuntimeCPU = UE::NNE::GetRuntime<INNERuntimeCPU>(RuntimeName);
 
 		if (ensureMsgf(RuntimeCPU.IsValid(), TEXT("Could not find requested NNE Runtime")))
diff --git a/Engine/Plugins/Animation/MLDeformer/NearestNeighborModel/Source/NearestNeighborModel/Private/NearestNeighborOptimizedNetwork.h b/Engine/Plugins/Animation/MLDeformer/NearestNeighborModel/Source/NearestNeighborModel/Private/NearestNeighborOptimizedNetwork.h
--- a/Engine/Plugins/Animation/MLDeformer/NearestNeighborModel/Source/NearestNeighborModel/Private/NearestNeighborOptimizedNetwork.h
+++ b/Engine/Plugins/Animation/MLDeformer/NearestNeighborModel/Source/NearestNeighborModel/Private/NearestNeighborOptimizedNetwork.h
@@ -160,7 +160,7 @@ protected:
 private:
 
 	static const TCHAR* RuntimeName;
-
+	static const TCHAR* RuntimeModuleName;
 };
 
 /** 
diff --git a/Engine/Plugins/Animation/MLDeformer/NeuralMorphModel/Source/NeuralMorphModel/Private/NeuralMorphNetwork.cpp b/Engine/Plugins/Animation/MLDeformer/NeuralMorphModel/Source/NeuralMorphModel/Private/NeuralMorphNetwork.cpp
--- a/Engine/Plugins/Animation/MLDeformer/NeuralMorphModel/Source/NeuralMorphModel/Private/NeuralMorphNetwork.cpp
+++ b/Engine/Plugins/Animation/MLDeformer/NeuralMorphModel/Source/NeuralMorphModel/Private/NeuralMorphNetwork.cpp
@@ -362,6 +362,8 @@ void UNeuralMorphNetwork::PostLoad()
 
 	// Create models
 
+	ensureMsgf(FModuleManager::Get().LoadModule(TEXT("NNERuntimeBasicCpu")), TEXT("Unable to load module for NNE runtime NNERuntimeBasicCpu."));
+	
 	TWeakInterfacePtr<INNERuntimeCPU> RuntimeCPU = UE::NNE::GetRuntime<INNERuntimeCPU>(TEXT("NNERuntimeBasicCpu"));
 
 	if (ensureMsgf(RuntimeCPU.IsValid(), TEXT("Could not find requested NNE Runtime")))
@@ -489,6 +491,8 @@ bool UNeuralMorphNetwork::Load(const FString& Filename)
 		FString RuntimeName;
 		UE::NeuralMorphModel::Private::Serialization::Load(Offset, RuntimeName, FileData);
 
+		ensureMsgf(RuntimeName == TEXT("NNERuntimeBasicCpu"), TEXT("Currently only NNERuntimeBasicCpu runtime is supported"));
+
 		// Load Main Network
 
 		if (!MainModelData)
@@ -522,7 +526,9 @@ bool UNeuralMorphNetwork::Load(const FString& Filename)
 
 		// Create models
 
-		TWeakInterfacePtr<INNERuntimeCPU> RuntimeCPU = UE::NNE::GetRuntime<INNERuntimeCPU>(RuntimeName);
+		ensureMsgf(FModuleManager::Get().LoadModule(TEXT("NNERuntimeBasicCpu")), TEXT("Unable to load module for NNE runtime NNERuntimeBasicCpu."));
+
+		TWeakInterfacePtr<INNERuntimeCPU> RuntimeCPU = UE::NNE::GetRuntime<INNERuntimeCPU>(TEXT("NNERuntimeBasicCpu"));
 
 		if (ensureMsgf(RuntimeCPU.IsValid(), TEXT("Could not find requested NNE Runtime")))
 		{
diff --git a/Engine/Plugins/Experimental/LearningAgents/Source/LearningAgents/Private/LearningAgentsNeuralNetworkData.cpp b/Engine/Plugins/Experimental/LearningAgents/Source/LearningAgents/Private/LearningAgentsNeuralNetworkData.cpp
--- a/Engine/Plugins/Experimental/LearningAgents/Source/LearningAgents/Private/LearningAgentsNeuralNetworkData.cpp
+++ b/Engine/Plugins/Experimental/LearningAgents/Source/LearningAgents/Private/LearningAgentsNeuralNetworkData.cpp
@@ -4,7 +4,9 @@
 
 #include "LearningAgentsNeuralNetwork.h"
 #include "LearningNeuralNetwork.h"
+
 #include "Async/ParallelFor.h"
+#include "Modules/ModuleManager.h"
 
 #include "NNE.h"
 #include "NNEModelData.h"
@@ -135,6 +137,8 @@ namespace UE::Learning::Agents
 
 		Parent.ModelData->Init(TEXT("ubnne"), Parent.FileData);
 
+		ensureMsgf(FModuleManager::Get().LoadModule(TEXT("NNERuntimeBasicCpu")), TEXT("Unable to load module for NNE runtime NNERuntimeBasicCpu."));
+
 		TWeakInterfacePtr<INNERuntimeCPU> RuntimeCPU = NNE::GetRuntime<INNERuntimeCPU>(TEXT("NNERuntimeBasicCpu"));
 
 		if (ensureMsgf(RuntimeCPU.IsValid(), TEXT("Could not find requested NNE Runtime")))
