Fix duplicate skinning code.

diff --git a/Engine/Plugins/Animation/MLDeformer/NearestNeighborModel/Shaders/Private/DataInterfaceSkeletonWithQuats.ush b/Engine/Plugins/Animation/MLDeformer/NearestNeighborModel/Shaders/Private/DataInterfaceSkeletonWithQuats.ush
--- a/Engine/Plugins/Animation/MLDeformer/NearestNeighborModel/Shaders/Private/DataInterfaceSkeletonWithQuats.ush
+++ b/Engine/Plugins/Animation/MLDeformer/NearestNeighborModel/Shaders/Private/DataInterfaceSkeletonWithQuats.ush
@@ -20,52 +20,12 @@ uint ReadNumVertices_{DataInterfaceName}()
 
 int4 GetBlendIndices_{DataInterfaceName}(uint StreamOffset, uint ExtraInfluenceIndex)
 {
-	int4 UnpackedBlendIndices = 0;
-
-#if GPUSKIN_BONE_INDEX_UINT16
-	StreamOffset += GPUSKIN_VB_OFFSET_INFLUENCEBONES + 2 * ExtraInfluenceIndex;
-	uint PackedBlendIndices = {DataInterfaceName}_InputWeightStream[StreamOffset];
-	UnpackedBlendIndices.x = PackedBlendIndices & 0xffff;
-	UnpackedBlendIndices.y = PackedBlendIndices >> 16 & 0xffff;
-	PackedBlendIndices = {DataInterfaceName}_InputWeightStream[StreamOffset + 1];
-	UnpackedBlendIndices.z = PackedBlendIndices & 0xffff;
-	UnpackedBlendIndices.w = PackedBlendIndices >> 16 & 0xffff;
-#else // GPUSKIN_BONE_INDEX_UINT16
-	StreamOffset += GPUSKIN_VB_OFFSET_INFLUENCEBONES + ExtraInfluenceIndex;
-	uint PackedBlendIndices = {DataInterfaceName}_InputWeightStream[StreamOffset];
-	UnpackedBlendIndices.x = PackedBlendIndices & 0xff;
-	UnpackedBlendIndices.y = PackedBlendIndices >> 8 & 0xff;
-	UnpackedBlendIndices.z = PackedBlendIndices >> 16 & 0xff;
-	UnpackedBlendIndices.w = PackedBlendIndices >> 24 & 0xff;
-#endif // GPUSKIN_BONE_INDEX_UINT16
-
-	return UnpackedBlendIndices;
+	return GetBlendIndices({DataInterfaceName}_InputWeightStream, StreamOffset, ExtraInfluenceIndex);
 }
 
 float4 GetBlendWeights_{DataInterfaceName}(uint StreamOffset, uint ExtraInfluenceIndex)
 {
-	StreamOffset += (({DataInterfaceName}_NumBoneInfluences + 3) / 4) * (1 + GPUSKIN_BONE_INDEX_UINT16);
-
-	float4 UnpackedBlendWeights = 0;
-#if GPUSKIN_BONE_WEIGHTS_UINT16
-	StreamOffset += ExtraInfluenceIndex * 2;
-	uint PackedBlendWeights = {DataInterfaceName}_InputWeightStream[StreamOffset];
-	UnpackedBlendWeights.x = float(PackedBlendWeights & 0xffff) / 65535.0f;
-	UnpackedBlendWeights.y = float(PackedBlendWeights >> 16 & 0xffff) / 65535.0f;
-	PackedBlendWeights = {DataInterfaceName}_InputWeightStream[StreamOffset + 1];
-	UnpackedBlendWeights.z = float(PackedBlendWeights & 0xffff) / 65535.0f;
-	UnpackedBlendWeights.w = float(PackedBlendWeights >> 16 & 0xffff) / 65535.0f;
-	
-#else
-	StreamOffset += ExtraInfluenceIndex;
-	uint PackedBlendWeights = {DataInterfaceName}_InputWeightStream[StreamOffset];
-	
-	UnpackedBlendWeights.x = float(PackedBlendWeights & 0xff) / 255.0f;
-	UnpackedBlendWeights.y = float(PackedBlendWeights >> 8 & 0xff) / 255.0f;
-	UnpackedBlendWeights.z = float(PackedBlendWeights >> 16 & 0xff) / 255.0f;
-	UnpackedBlendWeights.w = float(PackedBlendWeights >> 24 & 0xff) / 255.0f;
-#endif	
-	return UnpackedBlendWeights;
+	return GetBlendWeights({DataInterfaceName}_InputWeightStream, StreamOffset, ExtraInfluenceIndex, {DataInterfaceName}_NumBoneInfluences);
 }
 
 #endif // GPUSKIN_UNLIMITED_BONE_INFLUENCE
@@ -159,65 +119,20 @@ float3x4 ReadWeightedBoneMatrix_{DataInterfaceName}(uint VertexIndex)
 
 #elif GPUSKIN_UNLIMITED_BONE_INFLUENCE
 
-	float3x4 BlendMatrix = (float3x4)0;
-
-	uint BufferIndex = {DataInterfaceName}_InputStreamStart + VertexIndex;
-	uint BlendOffsetCount = {DataInterfaceName}_InputWeightLookupStream[BufferIndex];
-	int NumBoneInfluencesLocal = BlendOffsetCount & 0xff;
-	int IndicesOffset = BlendOffsetCount >> 8;
-	int BoneIndexSize = {DataInterfaceName}_InputWeightIndexSize & 0xff;
-	int BoneWeightSize = {DataInterfaceName}_InputWeightIndexSize >> 8;
-	int WeightsOffset = IndicesOffset + (BoneIndexSize * NumBoneInfluencesLocal);
+	const uint BufferIndex 		= {DataInterfaceName}_InputStreamStart + VertexIndex;
+	const uint BlendOffsetCount = {DataInterfaceName}_InputWeightLookupStream[BufferIndex];
+	return ComputeBoneMatrixWithUnlimitedInfluences({DataInterfaceName}_BoneMatrices, {DataInterfaceName}_InputWeightStream, {DataInterfaceName}_InputWeightIndexSize, BlendOffsetCount);
 
-	for (int InfluenceIndex = 0; InfluenceIndex < NumBoneInfluencesLocal; InfluenceIndex++)
-	{
-		int BoneIndexOffset = IndicesOffset + (BoneIndexSize * InfluenceIndex);
-		int BoneIndex = {DataInterfaceName}_InputWeightStream[BoneIndexOffset];
-		if (BoneIndexSize > 1)
-		{
-			BoneIndex = {DataInterfaceName}_InputWeightStream[BoneIndexOffset + 1] << 8 | BoneIndex;
-		}
-		int BoneWeightOffset = WeightsOffset + (BoneWeightSize * InfluenceIndex);
-		float BoneWeight = float({ DataInterfaceName }_InputWeightStream[BoneWeightOffset]) / 255.0f;
-		if (BoneWeightSize > 1)
-		{
-			BoneWeight = float({DataInterfaceName}_InputWeightStream[BoneWeightOffset + 1] << 8 | {DataInterfaceName}_InputWeightStream[BoneWeightOffset]) / 65535.0;
-		}
-		BlendMatrix += BoneWeight * GetBoneMatrix_{DataInterfaceName}(BoneIndex);
-	}
-
-	return BlendMatrix;
-
-#else // !GPUSKIN_UNLIMITED_BONE_INFLUENCE
-
-	float3x4 BlendMatrix = (float3x4)0;
-
-	uint StreamOffset = {DataInterfaceName}_InputWeightStart + (VertexIndex * ({DataInterfaceName}_InputWeightStride / 4));
-
-	if ({DataInterfaceName}_NumBoneInfluences > 0)
-	{
-		int4 BlendIndices = GetBlendIndices_{DataInterfaceName}(StreamOffset, 0);
-		float4 BlendWeights = GetBlendWeights_{DataInterfaceName}(StreamOffset, 0);
-
-		BlendMatrix += BlendWeights.x * GetBoneMatrix_{DataInterfaceName}(BlendIndices.x);
-		BlendMatrix += BlendWeights.y * GetBoneMatrix_{DataInterfaceName}(BlendIndices.y);
-		BlendMatrix += BlendWeights.z * GetBoneMatrix_{DataInterfaceName}(BlendIndices.z);
-		BlendMatrix += BlendWeights.w * GetBoneMatrix_{DataInterfaceName}(BlendIndices.w);
-	}
-
-	// todo[CF]: make this a permutation when permutation support is working again?
-	if ({DataInterfaceName}_NumBoneInfluences > 4)
-	{
-		int4 BlendIndices = GetBlendIndices_{DataInterfaceName}(StreamOffset, 1);
-		float4 BlendWeights = GetBlendWeights_{DataInterfaceName}(StreamOffset, 1);
+#else
 
-		BlendMatrix += BlendWeights.x * GetBoneMatrix_{DataInterfaceName}(BlendIndices.x);
-		BlendMatrix += BlendWeights.y * GetBoneMatrix_{DataInterfaceName}(BlendIndices.y);
-		BlendMatrix += BlendWeights.z * GetBoneMatrix_{DataInterfaceName}(BlendIndices.z);
-		BlendMatrix += BlendWeights.w * GetBoneMatrix_{DataInterfaceName}(BlendIndices.w);
-	}
+	const uint StreamOffset = {DataInterfaceName}_InputWeightStart + (VertexIndex * ({DataInterfaceName}_InputWeightStride / 4));
 
-	return BlendMatrix;
+	FGPUSkinIndexAndWeight IndicesAndWeights = (FGPUSkinIndexAndWeight)0;
+	IndicesAndWeights.BlendIndices  = GetBlendIndices_{DataInterfaceName}(StreamOffset, 0);
+	IndicesAndWeights.BlendIndices2 = GetBlendIndices_{DataInterfaceName}(StreamOffset, 1);
+	IndicesAndWeights.BlendWeights  = GetBlendWeights_{DataInterfaceName}(StreamOffset, 0);
+	IndicesAndWeights.BlendWeights2 = GetBlendWeights_{DataInterfaceName}(StreamOffset, 1);
+	return ComputeBoneMatrixWithLimitedInfluences({DataInterfaceName}_BoneMatrices, IndicesAndWeights);
 
 #endif
 }
