MLDeformer NearestNeighborModel: Use new pip installer system rather than packaged python dependencies

diff --git a/Engine/Plugins/Animation/MLDeformer/NearestNeighborModel/NearestNeighborModel.uplugin b/Engine/Plugins/Animation/MLDeformer/NearestNeighborModel/NearestNeighborModel.uplugin
--- a/Engine/Plugins/Animation/MLDeformer/NearestNeighborModel/NearestNeighborModel.uplugin
+++ b/Engine/Plugins/Animation/MLDeformer/NearestNeighborModel/NearestNeighborModel.uplugin
@@ -51,5 +51,19 @@
 			"Name": "NNERuntimeBasicCpu",
 			"Enabled": true
 		}
+	],
+	"PythonRequirements":
+	[
+		{
+			"Platform": "All",
+			"Requirements":
+			[
+				"joblib==1.2.0",
+				"onnx==1.13.0",
+				"scipy==1.8.1",
+				"scikit-learn==1.2.1",
+				"threadpoolctl==3.1.0"
+			]
+		}
 	]
 }
\ No newline at end of file
