Nearestneighbormodel: refactor
Add NearestNeighborModelSection class
Move kmeans to tools menu
Move neighbor data to tools menu
Visualize nearest neighbor data from the timeline dropdown menu

diff --git a/Engine/Plugins/Animation/MLDeformer/MLDeformerFramework/Source/MLDeformerFrameworkEditor/Private/MLDeformerEditorModel.cpp b/Engine/Plugins/Animation/MLDeformer/MLDeformerFramework/Source/MLDeformerFrameworkEditor/Private/MLDeformerEditorModel.cpp
--- a/Engine/Plugins/Animation/MLDeformer/MLDeformerFramework/Source/MLDeformerFrameworkEditor/Private/MLDeformerEditorModel.cpp
+++ b/Engine/Plugins/Animation/MLDeformer/MLDeformerFramework/Source/MLDeformerFrameworkEditor/Private/MLDeformerEditorModel.cpp
@@ -384,6 +384,11 @@ namespace UE::MLDeformer
 		return Result;
 	}
 
+	bool FMLDeformerEditorModel::IsAnimIndexValid(int32 AnimIndex) const
+	{
+		return FMath::IsWithinInclusive<int32>(AnimIndex, 0, GetNumTrainingInputAnims());
+	}
+
 	void FMLDeformerEditorModel::UpdateMeshOffsetFactors()
 	{
 		// Set the default mesh translation offsets for our ground truth actors.
@@ -513,6 +518,15 @@ namespace UE::MLDeformer
 		EditorActors.Empty();
 	}
 
+	void FMLDeformerEditorModel::SetTimelineAnimNames(const TArray<TSharedPtr<FMLDeformerTrainingInputAnimName>>& AnimNames)
+	{
+		SMLDeformerTimeline* Timeline = GetEditor()->GetTimeSlider();
+		if (Timeline)
+		{
+			Timeline->SetTrainingAnimNames(AnimNames);
+		}
+	}
+
 	FMLDeformerEditorActor* FMLDeformerEditorModel::FindEditorActor(int32 ActorTypeID) const
 	{
 		for (FMLDeformerEditorActor* EditorActor : EditorActors)
@@ -1579,7 +1593,7 @@ namespace UE::MLDeformer
 
 	void FMLDeformerEditorModel::SetActiveTrainingInputAnimIndex(int32 Index)
 	{
-		check(FMath::IsWithinInclusive<int32>(Index, 0, GetNumTrainingInputAnims()));
+		check(IsAnimIndexValid(Index));
 		ActiveTrainingInputAnimIndex = Index;
 	}
 
@@ -2806,12 +2820,7 @@ namespace UE::MLDeformer
 			}
 		}
 
-		// Pass this to the timeline.
-		SMLDeformerTimeline* Timeline = GetEditor()->GetTimeSlider();
-		if (Timeline)
-		{
-			Timeline->SetTrainingAnimNames(NameList);
-		}
+		SetTimelineAnimNames(NameList);
 	}
 
 	void FMLDeformerEditorModel::InvalidateDeltas()
diff --git a/Engine/Plugins/Animation/MLDeformer/MLDeformerFramework/Source/MLDeformerFrameworkEditor/Private/MLDeformerGeomCacheEditorModel.cpp b/Engine/Plugins/Animation/MLDeformer/MLDeformerFramework/Source/MLDeformerFrameworkEditor/Private/MLDeformerGeomCacheEditorModel.cpp
--- a/Engine/Plugins/Animation/MLDeformer/MLDeformerFramework/Source/MLDeformerFrameworkEditor/Private/MLDeformerGeomCacheEditorModel.cpp
+++ b/Engine/Plugins/Animation/MLDeformer/MLDeformerFramework/Source/MLDeformerFrameworkEditor/Private/MLDeformerGeomCacheEditorModel.cpp
@@ -63,7 +63,7 @@ namespace UE::MLDeformer
 	UGeometryCache* FMLDeformerGeomCacheEditorModel::GetActiveGeometryCache() const
 	{
 		const int32 ActiveAnimIndex = GetActiveTrainingInputAnimIndex();
-		if (ActiveAnimIndex == INDEX_NONE || ActiveAnimIndex >= GetNumTrainingInputAnims())
+		if (ActiveAnimIndex == INDEX_NONE)
 		{
 			return nullptr;
 		}
@@ -148,11 +148,8 @@ namespace UE::MLDeformer
 		UGeometryCache* GeomCache = nullptr;
 		if (ActiveAnimIndex != INDEX_NONE)
 		{
-			if (ActiveAnimIndex < GetNumTrainingInputAnims())
-			{
-				FMLDeformerGeomCacheTrainingInputAnim* Anim = static_cast<FMLDeformerGeomCacheTrainingInputAnim*>(GetTrainingInputAnim(ActiveAnimIndex));
-				GeomCache = Anim ? Anim->GetGeometryCache() : nullptr;
-			}
+			FMLDeformerGeomCacheTrainingInputAnim* Anim = static_cast<FMLDeformerGeomCacheTrainingInputAnim*>(GetTrainingInputAnim(ActiveAnimIndex));
+			GeomCache = Anim ? Anim->GetGeometryCache() : nullptr;
 		}
 
 		const FLinearColor LabelColor = FMLDeformerEditorStyle::Get().GetColor("MLDeformer.TargetMesh.LabelColor");
diff --git a/Engine/Plugins/Animation/MLDeformer/MLDeformerFramework/Source/MLDeformerFrameworkEditor/Private/MLDeformerGeomCacheSampler.cpp b/Engine/Plugins/Animation/MLDeformer/MLDeformerFramework/Source/MLDeformerFrameworkEditor/Private/MLDeformerGeomCacheSampler.cpp
--- a/Engine/Plugins/Animation/MLDeformer/MLDeformerFramework/Source/MLDeformerFrameworkEditor/Private/MLDeformerGeomCacheSampler.cpp
+++ b/Engine/Plugins/Animation/MLDeformer/MLDeformerFramework/Source/MLDeformerFrameworkEditor/Private/MLDeformerGeomCacheSampler.cpp
@@ -44,7 +44,8 @@ namespace UE::MLDeformer
 		UMLDeformerGeomCacheModel* GeomCacheModel = Cast<UMLDeformerGeomCacheModel>(InModel->GetModel());
 		check(GeomCacheModel);
 
-		UGeometryCache* GeomCache = GeomCacheModel->GetTrainingInputAnims()[InAnimIndex].GetGeometryCache();
+		FMLDeformerGeomCacheTrainingInputAnim* GeomCacheAnim = static_cast<FMLDeformerGeomCacheTrainingInputAnim*>(EditorModel->GetTrainingInputAnim(InAnimIndex));
+		UGeometryCache* GeomCache = GeomCacheAnim->GetGeometryCache();
 		GeometryCacheComponent->SetGeometryCache(GeomCache);
 		GeometryCacheComponent->SetManualTick(true);
 		GeometryCacheComponent->SetVisibility(false);
diff --git a/Engine/Plugins/Animation/MLDeformer/MLDeformerFramework/Source/MLDeformerFrameworkEditor/Public/MLDeformerEditorModel.h b/Engine/Plugins/Animation/MLDeformer/MLDeformerFramework/Source/MLDeformerFrameworkEditor/Public/MLDeformerEditorModel.h
--- a/Engine/Plugins/Animation/MLDeformer/MLDeformerFramework/Source/MLDeformerFrameworkEditor/Public/MLDeformerEditorModel.h
+++ b/Engine/Plugins/Animation/MLDeformer/MLDeformerFramework/Source/MLDeformerFrameworkEditor/Public/MLDeformerEditorModel.h
@@ -32,6 +32,7 @@ class UGeometryCache;
 class UMorphTarget;
 class FMorphTargetVertexInfoBuffers;
 struct FMLDeformerTrainingInputAnim;
+struct FMLDeformerTrainingInputAnimName;
 
 /** Training process return codes. */
 UENUM()
@@ -868,7 +869,7 @@ namespace UE::MLDeformer
 		void InvalidateDeltas();
 
 	protected:
-		void CreateSamplers();
+		virtual void CreateSamplers();
 
 		void AddAnimatedBonesToBonesIncludeList(const UAnimSequence* AnimSequence);
 		void AddAnimatedCurvesToCurvesIncludeList(const UAnimSequence* AnimSequence);
@@ -895,6 +896,11 @@ namespace UE::MLDeformer
 		 */
 		void DeleteEditorActors();
 
+		/**
+		 * Set the names of the anim sequences in the timeline. 
+		 */
+		void SetTimelineAnimNames(const TArray<TSharedPtr<FMLDeformerTrainingInputAnimName>>& AnimNames);
+
 		/**
 		 * Perform some basic checks to see if the editor can be ready to train the model.
 		 * This checks whether there is an anim sequence, target mesh, and skeletal mesh.
@@ -985,6 +991,7 @@ namespace UE::MLDeformer
 		void UpdateCompareActorLabels();		
 
 		int32 CalcNumValidCompareActorsPriorTo(int32 CompareActorIndex) const;
+		virtual bool IsAnimIndexValid(int32 AnimIndex) const;
 
 	protected:
 		/** The runtime model associated with this editor model. */
diff --git a/Engine/Plugins/Animation/MLDeformer/NearestNeighborModel/Source/NearestNeighborModel/Private/NearestNeighborModel.cpp b/Engine/Plugins/Animation/MLDeformer/NearestNeighborModel/Source/NearestNeighborModel/Private/NearestNeighborModel.cpp
--- a/Engine/Plugins/Animation/MLDeformer/NearestNeighborModel/Source/NearestNeighborModel/Private/NearestNeighborModel.cpp
+++ b/Engine/Plugins/Animation/MLDeformer/NearestNeighborModel/Source/NearestNeighborModel/Private/NearestNeighborModel.cpp
@@ -1,27 +1,30 @@
 // Copyright Epic Games, Inc. All Rights Reserved.
 
 #include "NearestNeighborModel.h"
-#include "NearestNeighborModelInputInfo.h"
-#include "NearestNeighborModelInstance.h"
-#include "NearestNeighborOptimizedNetwork.h"
-#include "NearestNeighborOptimizedNetworkLoader.h"
-#include "MLDeformerComponent.h"
+#include "Algo/Copy.h"
+#include "Components/ExternalMorphSet.h"
+#include "Components/SkinnedMeshComponent.h"
 #include "Engine/SkeletalMesh.h"
 #include "GeometryCache.h"
+#include "HAL/PlatformFileManager.h"
 #include "Interfaces/IPluginManager.h"
+#include "Internationalization/Regex.h"
 #include "Misc/FileHelper.h"
-#include "Components/SkinnedMeshComponent.h"
-#include "Components/ExternalMorphSet.h"
-#include "Rendering/SkeletalMeshModel.h"
+#include "Misc/Optional.h"
+#include "MLDeformerComponent.h"
+#include "MLDeformerGeomCacheTrainingInputAnim.h"
+#include "NearestNeighborModelInputInfo.h"
+#include "NearestNeighborModelInstance.h"
+#include "NearestNeighborModelVizSettings.h"
+#include "NearestNeighborOptimizedNetwork.h"
 #include "Rendering/MorphTargetVertexInfoBuffers.h"
+#include "Rendering/SkeletalMeshModel.h"
 #include "ShaderCore.h"
-#include "UObject/UObjectGlobals.h"
+
 #include UE_INLINE_GENERATED_CPP_BY_NAME(NearestNeighborModel)
 
 #define LOCTEXT_NAMESPACE "UNearestNeighborModel"
 
-NEARESTNEIGHBORMODEL_API DEFINE_LOG_CATEGORY(LogNearestNeighborModel)
-
 namespace UE::NearestNeighborModel
 {
 	class NEARESTNEIGHBORMODEL_API FNearestNeighborModelModule
@@ -33,11 +36,7 @@ namespace UE::NearestNeighborModel
 			AddShaderSourceDirectoryMapping(TEXT("/Plugin/MLDeformer/NearestNeighborModel"), PluginShaderDir);
 		}
 	};
-}
-IMPLEMENT_MODULE(UE::NearestNeighborModel::FNearestNeighborModelModule, NearestNeighborModel)
 
-namespace UE::NearestNeighborModel
-{
 	bool bNearestNeighborModelUseOptimizedNetwork = true;
 	FAutoConsoleVariableRef CVarNearestNeighborModelUseOptimizedNetwork(
 		TEXT("p.NearestNeighborModel.UseOptimizedNetwork"),
@@ -45,60 +44,809 @@ namespace UE::NearestNeighborModel
 		TEXT("Whether to use the optimized network.")
 	);
 
+	struct FNearestNeighborModelCustomVersion
+	{
+		enum Type
+		{
+			BeforeCustomVersionWasAdded = 0,
+	
+			VersionPlusOne,
+			LatestVersion = VersionPlusOne - 1
+		};
+	
+		const static FGuid GUID;
+	};
+	const FGuid FNearestNeighborModelCustomVersion::GUID = FGuid(0xBA926256, 0x6C7947C4, 0x96D1E4F4, 0xCE985342);
+	FCustomVersionRegistration GRegisterNearestNeighborModelCustomVersion(FNearestNeighborModelCustomVersion::GUID, FNearestNeighborModelCustomVersion::LatestVersion, TEXT("NearestNeighborModelVer"));
+
+	namespace Private
+	{
+		TOptional<TArray<int32>> ParseIntegers(const FString& String)
+		{
+			if (String.IsEmpty())
+			{
+				TArray<int32> Empty;
+			    return Empty;
+			}
+			TOptional<TArray<int32>> None;
+			static const FRegexPattern AllowedCharsPattern(TEXT("^[-,0-9\\s]+$"));
+		
+			if (!FRegexMatcher(AllowedCharsPattern, String).FindNext())
+			{
+			    UE_LOG(LogNearestNeighborModel, Error, TEXT("Input contains invalid characters."));
+			    return None;
+			}
+		
+			static const FRegexPattern SingleNumberPattern(TEXT("^\\s*(\\d+)\\s*$"));
+			static const FRegexPattern RangePattern(TEXT("^\\s*(\\d+)\\s*-\\s*(\\d+)\\s*$"));
+		
+			TArray<FString> Segments;
+			TArray<int32> Result;
+		    String.ParseIntoArray(Segments, TEXT(","), true);
+		    for (const FString& Segment : Segments)
+		    {
+		    	bool bSegmentValid = false;
+		
+		    	FRegexMatcher SingleNumberMatcher(SingleNumberPattern, Segment);
+		    	if (SingleNumberMatcher.FindNext())
+		    	{
+		    	    const int32 SingleNumber = FCString::Atoi(*SingleNumberMatcher.GetCaptureGroup(1));
+		    	    Result.Add(SingleNumber);
+		    	    bSegmentValid = true;
+		    	}
+		    	else
+		    	{
+		    		FRegexMatcher RangeMatcher(RangePattern, Segment);
+		    		if (RangeMatcher.FindNext())
+		    		{
+		    		    const int32 RangeStart = FCString::Atoi(*RangeMatcher.GetCaptureGroup(1));
+		    		    const int32 RangeEnd = FCString::Atoi(*RangeMatcher.GetCaptureGroup(2));
+		
+		    		    for (int32 i = RangeStart; i <= RangeEnd; ++i)
+		    		    {
+		    		        Result.Add(i);
+		    		    }
+		    		    bSegmentValid = true;
+		    		}
+		    	}
+		    	
+		    	if (!bSegmentValid)
+		    	{
+		    	    UE_LOG(LogNearestNeighborModel, Error, TEXT("Invalid format in segment: %s"), *Segment);
+					return None;
+		    	}
+		    }
+		
+			return Result;
+		}
+
+		FString IntegersToFormattedString(const TArray<int32>& Array)
+		{
+		    if (Array.Num() == 0)
+		        return "";
+		
+		    TArray<int32> SortedArray = Array;
+		    SortedArray.Sort();
+		
+		    FString Result;
+		    int32 StartRange = SortedArray[0];
+		    int32 EndRange = SortedArray[0];
+		
+		    for (int32 i = 1; i < SortedArray.Num(); ++i)
+		    {
+		        if (SortedArray[i] == EndRange + 1)
+		        {
+		            EndRange = SortedArray[i];
+		        }
+		        else
+		        {
+		            if (StartRange == EndRange)
+		                Result += FString::Printf(TEXT("%d, "), StartRange);
+		            else
+		                Result += FString::Printf(TEXT("%d-%d, "), StartRange, EndRange);
+		
+		            StartRange = EndRange = SortedArray[i];
+		        }
+		    }
+		
+		    if (StartRange == EndRange)
+		        Result += FString::FromInt(StartRange);
+		    else
+		        Result += FString::Printf(TEXT("%d-%d"), StartRange, EndRange);
+		
+		    return Result;
+		}
+
+		template<typename SrcT, typename TgtT>
+		TArray<TgtT> ConvertArray(TConstArrayView<SrcT> Src)
+		{
+			TArray<TgtT> Tgt;
+			Tgt.SetNum(Src.Num());
+			for (int32 Index = 0; Index < Src.Num(); ++Index)
+			{
+				Tgt[Index] = static_cast<TgtT>(Src[Index]);
+			}
+			return Tgt;
+		}
+
+		FString GetArchitectureString(int32 InputDim, TConstArrayView<int32> HiddenLayerDims, int32 OutputDim)
+		{
+			FString HiddenLayerString; 
+			for (int32 Index = 0; Index < HiddenLayerDims.Num() - 1; ++Index)
+			{
+				HiddenLayerString += FString::Printf(TEXT("%d, "), HiddenLayerDims[Index]);
+			}
+			HiddenLayerString += FString::FromInt(HiddenLayerDims.Last());
+			return FString::Printf(TEXT("%d [%s] %d"), InputDim, *HiddenLayerString, OutputDim);
+		}
+
+		TArray<int32> GetDefaultHiddenLayerDims(int32 InputDim, int32 OutputDim)
+		{
+			check(InputDim > 0 && OutputDim > 0);
+			TArray<int32> HiddenLayerDims;
+			if (OutputDim >= InputDim)
+			{
+				int32 Dim = FMath::RoundUpToPowerOfTwo(InputDim);
+				do 
+				{
+					HiddenLayerDims.Add(Dim);
+					Dim *= 2;
+				} while (Dim < OutputDim);
+			}
+			else
+			{
+				int32 Dim = FMath::RoundUpToPowerOfTwo(OutputDim);
+				do 
+				{
+					HiddenLayerDims.Add(Dim);
+					Dim *= 2;
+				} while (Dim < InputDim);
+				Algo::Reverse(HiddenLayerDims);
+			}
+			return HiddenLayerDims;
+		}
+
+		TOptional<FDateTime> GetTimeStamp(const FString& Path)
+		{
+			IPlatformFile& PlatformFile = FPlatformFileManager::Get().GetPlatformFile();
+			if (PlatformFile.FileExists(*Path))
+			{
+				return PlatformFile.GetTimeStamp(*Path);
+			}
+			else
+			{
+				return TOptional<FDateTime>();
+			}
+		}
+
+		TOptional<FDateTime> GetLatestTimeStamp(const TArray<FString>& Paths)
+		{
+			TOptional<FDateTime> LatestTimeStamp;
+			for (const FString& Path : Paths)
+			{
+				TOptional<FDateTime> TimeStamp = GetTimeStamp(Path);
+				if (TimeStamp.IsSet())
+				{
+					if (!LatestTimeStamp.IsSet() || TimeStamp.GetValue() > LatestTimeStamp.GetValue())
+					{
+						LatestTimeStamp = TimeStamp;
+					}
+				}
+			}
+			return LatestTimeStamp;
+		}
+
+		TArray<int32> GetIncludedFrames(TConstArrayView<int32> ExcludedFrames, int32 NumFrames)
+		{
+			TArray<int32> IncludedFrames;
+			TArray<int32> FrameMap;
+			IncludedFrames.Reserve(NumFrames);
+			for (int32 Frame = 0; Frame < NumFrames; Frame++)
+			{
+				if (!ExcludedFrames.Contains(Frame))
+				{
+					IncludedFrames.Add(Frame);
+				}
+			}
+			return IncludedFrames;
+		}
+
+		bool IsArrayChanged(EPropertyChangeType::Type ChangeType)
+		{
+			return (ChangeType & (EPropertyChangeType::ArrayAdd | EPropertyChangeType::ArrayRemove | EPropertyChangeType::ArrayClear | EPropertyChangeType::ArrayMove)) != 0;
+		}
+
+#if WITH_EDITORONLY_DATA
+		TArray<FInt32Range> GetMeshVertRanges(const USkeletalMesh& SkelMesh)
+		{
+			const FSkeletalMeshModel* ImportedModel = SkelMesh.GetImportedModel();
+			TArray<FInt32Range> Empty;
+			if (!ImportedModel)
+			{
+				return Empty;
+			}
+			constexpr int32 LODIndex = 0;
+			if (!ImportedModel->LODModels.IsValidIndex(LODIndex))
+			{
+				return Empty;
+			}
+			const FSkeletalMeshLODModel& LODModel = ImportedModel->LODModels[LODIndex];
+			const TArray<FSkelMeshImportedMeshInfo>& SkelMeshInfos = LODModel.ImportedMeshInfos;
+			TArray<FInt32Range> MeshVertRanges;
+			MeshVertRanges.SetNum(SkelMeshInfos.Num());
+			for (int32 Index = 0; Index < SkelMeshInfos.Num(); ++Index)
+			{
+				const FSkelMeshImportedMeshInfo& Info = SkelMeshInfos[Index];
+				MeshVertRanges[Index] = FInt32Range(Info.StartImportedVertex, Info.StartImportedVertex + Info.NumVertices);
+			}
+			return MeshVertRanges;
+		}
+#endif
+	};
+}
+IMPLEMENT_MODULE(UE::NearestNeighborModel::FNearestNeighborModelModule, NearestNeighborModel)
+
+
+// ================================== UNearestNeighborModelSection ==================================
+#if WITH_EDITOR
+const UAnimSequence* UNearestNeighborModelSection::GetNeighborPoses() const
+{
+	return NeighborPoses;
+}
+
+UAnimSequence* UNearestNeighborModelSection::GetMutableNeighborPoses() const
+{
+	return NeighborPoses;
+}
+
+const UGeometryCache* UNearestNeighborModelSection::GetNeighborMeshes() const
+{
+	return NeighborMeshes;
+}
+
+UGeometryCache* UNearestNeighborModelSection::GetMutableNeighborMeshes() const
+{
+	return NeighborMeshes;
+}
+
+void UNearestNeighborModelSection::SetVertexMapString(const FString& InString)
+{
+	VertexMapString = InString;
+	UpdateVertexMap();
+	check(Model);
+	Model->InvalidateTrainingModelOnly();
+}
+
+int32 UNearestNeighborModelSection::GetAssetNumNeighbors() const
+{
+	using UE::NearestNeighborModel::FHelpers;
+	return FMath::Min(FHelpers::GetNumFrames(GetNeighborPoses()), FHelpers::GetNumFrames(GetNeighborMeshes()));
+}
+
+void UNearestNeighborModelSection::SetPCAData(const TArray<float>& InVertexMean, const TArray<float>& InPCABasis)
+{
+	VertexMean = InVertexMean;
+	PCABasis = InPCABasis;
+	InvalidateInference();
+}
+
+void UNearestNeighborModelSection::SetNeighborData(const TArray<float>& InAssetNeighborCoeffs, const TArray<float>& InAssetNeighborOffsets)
+{
+	check(NumPCACoeffs > 0);
+	const int32 TryNumNeighbors = InAssetNeighborCoeffs.Num() / NumPCACoeffs;
+	check(TryNumNeighbors * NumPCACoeffs == InAssetNeighborCoeffs.Num());
+	check(TryNumNeighbors * NumVertices * 3 == InAssetNeighborOffsets.Num());
+	AssetNeighborCoeffs = InAssetNeighborCoeffs;
+	AssetNeighborOffsets = InAssetNeighborOffsets;
+	InvalidateInference();
+}
+
+const TArray<int32>&  UNearestNeighborModelSection::GetVertexMap() const
+{
+	return VertexMap;
+}
+
+const TArray<float>& UNearestNeighborModelSection::GetPCABasis() const
+{
+	return PCABasis;
+}
+
+const TArray<float>& UNearestNeighborModelSection::GetVertexMean() const
+{
+	return VertexMean;
+}
+
+const TArray<float>& UNearestNeighborModelSection::GetAssetNeighborCoeffs() const
+{
+	return AssetNeighborCoeffs;
+}
+
+const TArray<float>& UNearestNeighborModelSection::GetAssetNeighborOffsets() const
+{
+	return AssetNeighborOffsets;
+}
+
+const TArray<int32>& UNearestNeighborModelSection::GetExcludedFrames() const
+{
+	return ExcludedFrames;
+}
+
+const TArray<int32>& UNearestNeighborModelSection::GetAssetNeighborIndexMap() const
+{
+	return AssetNeighborIndexMap;
+}
+
+FMLDeformerGeomCacheTrainingInputAnim* UNearestNeighborModelSection::GetInputAnim() const
+{
+	if (!InputAnim.IsValid())
+	{
+		InputAnim = MakeUnique<FMLDeformerGeomCacheTrainingInputAnim>();
+	}
+	InputAnim->SetAnimSequence(NeighborPoses);
+	InputAnim->SetGeometryCache(NeighborMeshes);
+	return InputAnim.Get();
+}
+
+bool UNearestNeighborModelSection::UpdateRuntimeNeighbors()
+{
+	if (NumPCACoeffs <= 0)
+	{
+		return false;
+	}
+	const int32 AssetNumNeighbors = AssetNeighborCoeffs.Num() / NumPCACoeffs;
+	if (AssetNumNeighbors * NumPCACoeffs != AssetNeighborCoeffs.Num())
+	{
+		return false;
+	}
+	const TArray<int32> IncludedFrames = UE::NearestNeighborModel::Private::GetIncludedFrames(ExcludedFrames, AssetNumNeighbors);
+	RuntimeNumNeighbors = IncludedFrames.Num();
+	if (RuntimeNumNeighbors <= 0)
+	{
+		return true;
+	}
+	RuntimeNeighborCoeffs.Reset();
+	RuntimeNeighborCoeffs.Reserve(RuntimeNumNeighbors * NumPCACoeffs);	
+	for (int32 Index = 0; Index < IncludedFrames.Num(); ++Index)
+	{
+		const int32 Frame = IncludedFrames[Index];
+		TConstArrayView<float> AssetCoeffs(AssetNeighborCoeffs.GetData() + Frame * NumPCACoeffs,  NumPCACoeffs);
+		Algo::Copy(AssetCoeffs, RuntimeNeighborCoeffs);
+	}
+	AssetNeighborIndexMap = IncludedFrames;
+	return true;
+}
+
+bool UNearestNeighborModelSection::GetRuntimeNeighborOffsets(TArray<float>& OutNeighborOffsets) const
+{
+	if (!IsReadyForInference())
+	{
+		return false;
+	}
+	const int32 AssetNumNeighbors = GetAssetNumNeighbors();
+	check(AssetNumNeighbors * NumVertices * 3 == AssetNeighborOffsets.Num());
+	OutNeighborOffsets.Reset();
+	OutNeighborOffsets.Reserve(RuntimeNumNeighbors * NumVertices * 3);
+	const TArray<int32> IncludedFrames = UE::NearestNeighborModel::Private::GetIncludedFrames(ExcludedFrames, AssetNumNeighbors);
+	check(IncludedFrames.Num() == RuntimeNumNeighbors);
+	for (int32 Index = 0; Index < IncludedFrames.Num(); ++Index)
+	{
+		const int32 Frame = IncludedFrames[Index];
+		TConstArrayView<float> AssetOffsets(AssetNeighborOffsets.GetData() + Frame * NumVertices * 3, NumVertices * 3);
+		Algo::Copy(AssetOffsets, OutNeighborOffsets);
+	}
+	return true;
+}
+#endif
+
+int32 UNearestNeighborModelSection::GetNumPCACoeffs() const
+{
+	return NumPCACoeffs;
+}
+
+int32 UNearestNeighborModelSection::GetNumVertices() const
+{
+	return NumVertices;
+}
+
+int32 UNearestNeighborModelSection::GetRuntimeNumNeighbors() const
+{
+	return RuntimeNumNeighbors;
+}
+
+const TArray<float>& UNearestNeighborModelSection::GetNeighborCoeffs() const
+{
+	return RuntimeNeighborCoeffs;
+}
+
+bool UNearestNeighborModelSection::IsReadyForTraining() const
+{
+	return bIsReadyForTraining;
+}
+
+bool UNearestNeighborModelSection::IsReadyForInference() const
+{
+	return bIsReadyForInference;
+}
+
+#if WITH_EDITOR
+UE::NearestNeighborModel::EOpFlag UNearestNeighborModelSection::UpdateForTraining()
+{
+	if (bIsReadyForTraining)
+	{
+		return EOpFlag::Success;
+	}
+	if (!GetModel())
+	{
+		UE_LOG(LogNearestNeighborModel, Error, TEXT("The parent is not a NearestNeighborModel."));
+		return EOpFlag::Error;
+	}
+
+	const EOpFlag Result = UpdateVertexMap();
+	using namespace UE::NearestNeighborModel;
+	if (!OpFlag::HasError(Result))
+	{
+		bIsReadyForTraining = true;
+	}
+	return Result;
+}
+
+UE::NearestNeighborModel::EOpFlag UNearestNeighborModelSection::UpdateForInference()
+{
+	if (bIsReadyForInference)
+	{
+		return EOpFlag::Success;
+	}
+
+	EOpFlag Result = EOpFlag::Success;
+	if (!IsReadyForTraining())
+	{
+		Result |= UpdateForTraining();
+	}
+	using namespace UE::NearestNeighborModel;
+	if (OpFlag::HasError(Result))
+	{
+		return Result;
+	}
+
+	if (!IsPCAValid())
+	{
+		UE_LOG(LogNearestNeighborModel, Error, TEXT("PCA data is invalid. Please re-train your model"));
+		ResetPCAData();
+		return EOpFlag::Error;
+	}
+	if (!UpdateRuntimeNeighbors())
+	{
+		UE_LOG(LogNearestNeighborModel, Error, TEXT("Failed to update runtime neighbors."));
+		return EOpFlag::Error;
+	}
+	if (!IsNearestNeighborValid())
+	{
+		UE_LOG(LogNearestNeighborModel, Error, TEXT("Nearest neighbor data is invalid. Please re-train your model to update PCA data or re-update your model to update nearest neighbor data."));
+		ResetNearestNeighborData();
+		return EOpFlag::Error;
+	}
+	bIsReadyForInference = true;
+	return Result;
+}
+
+void UNearestNeighborModelSection::InvalidateTraining()
+{
+	bIsReadyForTraining = bIsReadyForInference = false;
+	if (Model)
+	{
+		Model->InvalidateTrainingModelOnly();
+	}
+}
+
+void UNearestNeighborModelSection::InvalidateInference()
+{
+	bIsReadyForInference = false;
+	if (Model)
+	{
+		Model->InvalidateInferenceModelOnly();
+	}
+}
+
+void UNearestNeighborModelSection::SetModel(UNearestNeighborModel* InModel)
+{
+	if (Model != InModel)
+	{
+		Model = InModel;
+		Reset();
+	}
+}
+
+const UNearestNeighborModel* UNearestNeighborModelSection::GetModel() const
+{
+	return Model;
+}
+
+void UNearestNeighborModelSection::PostEditChangeProperty(FPropertyChangedEvent& PropertyChangedEvent)
+{
+	const FProperty* Property = PropertyChangedEvent.Property;
+	if (!Property)
+	{
+		return;
+	}
+	if (!Model)
+	{
+		return;
+	}
+	if (Property->GetFName() == GET_MEMBER_NAME_CHECKED(UNearestNeighborModelSection, NumPCACoeffs))
+	{
+		Model->UpdateNetworkOutputDim();
+		InvalidateTraining();
+	}
+	if (Property->GetFName() == GET_MEMBER_NAME_CHECKED(UNearestNeighborModelSection, VertexMapString))
+	{
+		InvalidateTraining();
+		UpdateVertexMap();
+	}
+	if (Property->GetFName() == GET_MEMBER_NAME_CHECKED(UNearestNeighborModelSection, NeighborPoses) ||
+		Property->GetFName() == GET_MEMBER_NAME_CHECKED(UNearestNeighborModelSection, NeighborMeshes) ||
+		Property->GetFName() == GET_MEMBER_NAME_CHECKED(UNearestNeighborModelSection, ExcludedFrames))
+	{
+		InvalidateInference();
+	}
+	Super::PostEditChangeProperty(PropertyChangedEvent);
+	Model->PostEditChangeProperty(PropertyChangedEvent);
+}
+
+int32 UNearestNeighborModelSection::GetMeshIndex() const
+{
+	return MeshIndex;
+}
+
+void UNearestNeighborModelSection::SetMeshIndex(int32 Index)
+{
+	MeshIndex = Index;
+	check(Model);
+	Model->InvalidateTrainingModelOnly();
+}
+
+bool UNearestNeighborModelSection::IsPCAValid() const
+{
+	return NumPCACoeffs >= 1 && NumVertices > 0 
+	&& VertexMean.Num() == NumVertices * 3 && PCABasis.Num() == NumVertices * 3 * NumPCACoeffs;
+}
+
+bool UNearestNeighborModelSection::IsPCAEmpty() const
+{
+	return VertexMean.IsEmpty() && PCABasis.IsEmpty();
+}
+
+void UNearestNeighborModelSection::ResetPCAData()
+{
+	VertexMean.Reset();
+	PCABasis.Reset();
+}
+
+bool UNearestNeighborModelSection::IsNearestNeighborValid() const
+{
+	return IsNearestNeighborEmpty() 
+	|| (RuntimeNeighborCoeffs.Num() == RuntimeNumNeighbors * NumPCACoeffs && NumVertices > 0);
+}
+
+bool UNearestNeighborModelSection::IsNearestNeighborEmpty() const
+{
+	return RuntimeNumNeighbors == 0 && RuntimeNeighborCoeffs.IsEmpty();
+}
+
+void UNearestNeighborModelSection::Reset()
+{
+	NumVertices = 0;
+	RuntimeNumNeighbors = 0;
+	VertexMap.Reset();
+	ResetPCAData();
+	ResetNearestNeighborData();
+	bIsReadyForTraining = false;
+	bIsReadyForInference = false;
+}
+
+void UNearestNeighborModelSection::ResetNearestNeighborData()
+{
+	AssetNeighborCoeffs.Reset();
+	AssetNeighborOffsets.Reset();
+	AssetNeighborIndexMap.Reset();
+	RuntimeNumNeighbors = 0;
+	RuntimeNeighborCoeffs.Reset();
+}
+
+UE::NearestNeighborModel::EOpFlag UNearestNeighborModelSection::UpdateVertexMap()
+{
+	using namespace UE::NearestNeighborModel;
+	if (!GetModel())
+	{
+		UE_LOG(LogNearestNeighborModel, Error, TEXT("NearestNeighborModel is invalid."));
+		return EOpFlag::Error;
+	}
+	const USkeletalMesh* const SkeletalMesh = GetModel()->GetSkeletalMesh();
+	if (!SkeletalMesh)
+	{
+		UE_LOG(LogNearestNeighborModel, Error, TEXT("SkeletalMesh is None"));
+		return EOpFlag::Error;
+	}
+	if (!SkeletalMesh->GetImportedModel())
+	{
+		UE_LOG(LogNearestNeighborModel, Error, TEXT("SkeletalMesh has no imported model."));
+		return EOpFlag::Error;
+	}
+	constexpr int32 LODIndex = 0;
+	if (!SkeletalMesh->GetImportedModel()->LODModels.IsValidIndex(LODIndex))
+	{
+		UE_LOG(LogNearestNeighborModel, Error, TEXT("SkeletalMesh has no LODModel."));
+		return EOpFlag::Error;
+	}
+	const FSkeletalMeshLODModel& LODModel = SkeletalMesh->GetImportedModel()->LODModels[LODIndex];
+
+	using namespace Private;
+	TOptional<TArray<int32>> VertexMapOpt = ParseIntegers(VertexMapString);
+	if (VertexMapOpt.IsSet())
+	{
+		VertexMap = MoveTemp(VertexMapOpt.GetValue());
+		// TODO: check uniqueness
+		if (const int32 MaxIndex = FMath::Max(VertexMap); MaxIndex >= SkeletalMesh->GetNumImportedVertices())
+		{
+			UE_LOG(LogNearestNeighborModel, Error, TEXT("MaxIndex %d is invalid, SkeletalMesh has %d vertices."), MaxIndex, SkeletalMesh->GetNumImportedVertices());
+			return EOpFlag::Error;
+		}
+	}
+	else
+	{
+		UE_LOG(LogNearestNeighborModel, Error, TEXT("VertexMapString is invalid."));
+		return EOpFlag::Error;
+	}
+
+	NumVertices = VertexMap.Num();
+	if (NumVertices <= NumPCACoeffs)
+	{
+		UE_LOG(LogNearestNeighborModel, Error, TEXT("NumVertices %d needs to be larger than NumPCACoeffs %d."), NumVertices, NumPCACoeffs);
+		return EOpFlag::Error;
+	}
+	return EOpFlag::Success;
+}
+FString UNearestNeighborModelSection::BackwardCompatibleVertexString = TEXT("Saved with backward version");
+#endif
+
+PRAGMA_DISABLE_DEPRECATION_WARNINGS 
+#if WITH_EDITORONLY_DATA
+void UNearestNeighborModelSection::InitFromClothPartData(FClothPartData& InPart)
+{
+	using UE::NearestNeighborModel::Private::ConvertArray;
+	VertexMap = ConvertArray<uint32, int32>(MoveTemp(InPart.VertexMap));
+	using UE::NearestNeighborModel::Private::IntegersToFormattedString;
+	VertexMapString = IntegersToFormattedString(VertexMap);
+	MeshIndex = INDEX_NONE;
+	NumPCACoeffs = InPart.PCACoeffNum;
+	PCABasis = MoveTemp(InPart.PCABasis);
+	VertexMean = MoveTemp(InPart.VertexMean);
+	NumVertices = VertexMap.Num();
+	RuntimeNumNeighbors = InPart.NumNeighbors;
+	// When loading from older versions, set RuntimeNeighborCoeffs to AssetNeighborCoeffs
+	AssetNeighborCoeffs = MoveTemp(InPart.NeighborCoeffs);
+	RuntimeNeighborCoeffs = AssetNeighborCoeffs;
+}
+#endif
+PRAGMA_ENABLE_DEPRECATION_WARNINGS
+
+// ================================== UNearestNeighborModel ==================================
+UNearestNeighborModel::UNearestNeighborModel(const FObjectInitializer& ObjectInitializer)
+	: Super(ObjectInitializer)
+{
+#if WITH_EDITORONLY_DATA
+	SetVizSettings(ObjectInitializer.CreateEditorOnlyDefaultSubobject<UNearestNeighborModelVizSettings>(this, TEXT("VizSettings")));
+#endif
+}
+
+void UNearestNeighborModel::GetAssetRegistryTags(TArray<FAssetRegistryTag>& OutTags) const
+{
+	Super::GetAssetRegistryTags(OutTags);
+
+	#if WITH_EDITORONLY_DATA
+		OutTags.Add(FAssetRegistryTag("MLDeformer.NearestNeighborModel.NumEpochs", FString::FromInt(NumEpochs), FAssetRegistryTag::TT_Numerical));
+		OutTags.Add(FAssetRegistryTag("MLDeformer.NearestNeighborModel.BatchSize", FString::FromInt(BatchSize), FAssetRegistryTag::TT_Numerical));
+		OutTags.Add(FAssetRegistryTag("MLDeformer.NearestNeighborModel.NumHiddenLayers", FString::FromInt(HiddenLayerDims.Num()), FAssetRegistryTag::TT_Numerical));
+		OutTags.Add(FAssetRegistryTag("MLDeformer.NearestNeighborModel.LearningRate", FString::Printf(TEXT("%f"), LearningRate), FAssetRegistryTag::TT_Numerical));
+		OutTags.Add(FAssetRegistryTag("MLDeformer.NearestNeighborModel.EarlyStopEpochs", FString::FromInt(EarlyStopEpochs), FAssetRegistryTag::TT_Numerical));
+		OutTags.Add(FAssetRegistryTag("MLDeformer.NearestNeighborModel.UseDualQuaternions", bUseDualQuaternionDeltas ? TEXT("True") : TEXT("False"), FAssetRegistryTag::TT_Alphabetical));
+	#endif
+
+	OutTags.Add(FAssetRegistryTag("MLDeformer.NearestNeighborModel.DecayFactor", FString::Printf(TEXT("%f"), DecayFactor), FAssetRegistryTag::TT_Numerical));
+	OutTags.Add(FAssetRegistryTag("MLDeformer.NearestNeighborModel.NearestNeighborOffsetWeight", FString::Printf(TEXT("%f"), NearestNeighborOffsetWeight), FAssetRegistryTag::TT_Numerical));
+	OutTags.Add(FAssetRegistryTag("MLDeformer.NearestNeighborModel.UseRBF", bUseRBF ? TEXT("True") : TEXT("False"), FAssetRegistryTag::TT_Alphabetical));
+	OutTags.Add(FAssetRegistryTag("MLDeformer.NearestNeighborModel.RBFSigma", FString::Printf(TEXT("%f"), RBFSigma), FAssetRegistryTag::TT_Numerical));
+}
+
+#if WITH_EDITOR
+void UNearestNeighborModel::PostEditChangeProperty(FPropertyChangedEvent& PropertyChangedEvent)
+{
+	const FProperty* Property = PropertyChangedEvent.Property;
+	if (!Property)
+	{
+		Super::PostEditChangeProperty(PropertyChangedEvent);
+		return;
+	}
+	if (Property->GetFName() == UMLDeformerModel::GetSkeletalMeshPropertyName() ||
+		PropertyChangedEvent.GetMemberPropertyName() == TEXT("TrainingInputAnims") ||
+		Property->GetFName() == UMLDeformerModel::GetAlignmentTransformPropertyName() ||
+		Property->GetFName() == UMLDeformerModel::GetBoneIncludeListPropertyName() ||
+		Property->GetFName() == UMLDeformerModel::GetCurveIncludeListPropertyName() ||
+		Property->GetFName() == UMLDeformerModel::GetMaxTrainingFramesPropertyName() ||
+		Property->GetFName() == UMLDeformerModel::GetDeltaCutoffLengthPropertyName() ||
+		Property->GetFName() == UNearestNeighborModel::GetHiddenLayerDimsPropertyName() ||
+		Property->GetFName() == UNearestNeighborModel::GetNumEpochsPropertyName() ||
+		Property->GetFName() == UNearestNeighborModel::GetBatchSizePropertyName() ||
+		Property->GetFName() == UNearestNeighborModel::GetLearningRatePropertyName() ||
+		Property->GetFName() == UNearestNeighborModel::GetEarlyStopEpochsPropertyName())
+	{
+		InvalidateTraining();
+	}
+
+	if (Property->GetFName() == UNearestNeighborModel::GetBoneIncludeListPropertyName() ||
+		Property->GetFName() == UNearestNeighborModel::GetCurveIncludeListPropertyName())
+	{
+		UpdateNetworkInputDim();
+	}
+	if (Property->GetFName() == UNearestNeighborModel::GetSectionsPropertyName())
+	{
+		if (PropertyChangedEvent.ChangeType == EPropertyChangeType::ArrayAdd)
+		{
+			const int32 NewIndex = PropertyChangedEvent.GetArrayIndex(UNearestNeighborModel::GetSectionsPropertyName().ToString());
+			FSection* const Section = OnSectionAdded(NewIndex);
+			check(Section);
+
+			FString MapString;
+			if (GetSkeletalMesh())
+			{
+				const TArray<FInt32Range> Ranges = UE::NearestNeighborModel::Private::GetMeshVertRanges(*GetSkeletalMesh());
+				const int32 MeshIndex = Section->GetMeshIndex(); 
+				if (Ranges.IsValidIndex(MeshIndex))
+				{
+					MapString = FString::Printf(TEXT("%d-%d"), Ranges[MeshIndex].GetLowerBoundValue(), Ranges[MeshIndex].GetUpperBoundValue() - 1);
+				}
+			}
+			Section->SetVertexMapString(MapString);
+		}
+		if (UE::NearestNeighborModel::Private::IsArrayChanged(PropertyChangedEvent.ChangeType))
+		{
+			InvalidateTrainingModelOnly();
+			UpdateNetworkOutputDim();
+		}
+	}
+	if (Property->GetFName() == UNearestNeighborModel::GetUseFileCachePropertyName() ||
+		Property->GetFName() == UNearestNeighborModel::GetFileCacheDirectoryPropertyName())
+	{
+		UpdateFileCache();
+	}
+	Super::PostEditChangeProperty(PropertyChangedEvent);
 }
+#endif
 
-UNearestNeighborModel::UNearestNeighborModel(const FObjectInitializer& ObjectInitializer)
-	: Super(ObjectInitializer)
+UMLDeformerInputInfo* UNearestNeighborModel::CreateInputInfo()
 {
-#if WITH_EDITORONLY_DATA
-	SetVizSettings(ObjectInitializer.CreateEditorOnlyDefaultSubobject<UNearestNeighborModelVizSettings>(this, TEXT("VizSettings")));
-#endif
+	return NewObject<UNearestNeighborModelInputInfo>(this);
 }
 
-void UNearestNeighborModel::GetAssetRegistryTags(TArray<FAssetRegistryTag>& OutTags) const
+FString UNearestNeighborModel::GetDisplayName() const
 {
-	Super::GetAssetRegistryTags(OutTags);
-
-	#if WITH_EDITORONLY_DATA
-		OutTags.Add(FAssetRegistryTag("MLDeformer.NearestNeighborModel.NumEpochs", FString::FromInt(NumEpochs), FAssetRegistryTag::TT_Numerical));
-		OutTags.Add(FAssetRegistryTag("MLDeformer.NearestNeighborModel.BatchSize", FString::FromInt(BatchSize), FAssetRegistryTag::TT_Numerical));
-		OutTags.Add(FAssetRegistryTag("MLDeformer.NearestNeighborModel.BasisSmoothIter", FString::FromInt(BasisSmoothIter), FAssetRegistryTag::TT_Numerical));
-		OutTags.Add(FAssetRegistryTag("MLDeformer.NearestNeighborModel.NumHiddenLayers", FString::FromInt(HiddenLayerDims.Num()), FAssetRegistryTag::TT_Numerical));
-		OutTags.Add(FAssetRegistryTag("MLDeformer.NearestNeighborModel.NumClusters", FString::FromInt(NumClusters), FAssetRegistryTag::TT_Numerical));
-		OutTags.Add(FAssetRegistryTag("MLDeformer.NearestNeighborModel.NumClusterAnims", FString::FromInt(SourceAnims.Num()), FAssetRegistryTag::TT_Numerical));
-		OutTags.Add(FAssetRegistryTag("MLDeformer.NearestNeighborModel.NumNearestNeighborAnims", FString::FromInt(NearestNeighborData.Num()), FAssetRegistryTag::TT_Numerical));
-		OutTags.Add(FAssetRegistryTag("MLDeformer.NearestNeighborModel.KMeansPartID", FString::FromInt(KMeansPartId), FAssetRegistryTag::TT_Numerical));
-		OutTags.Add(FAssetRegistryTag("MLDeformer.NearestNeighborModel.LearningRate", FString::Printf(TEXT("%f"), LearningRate), FAssetRegistryTag::TT_Numerical));
-		OutTags.Add(FAssetRegistryTag("MLDeformer.NearestNeighborModel.IsPartDataValid", bClothPartDataValid ? TEXT("True") : TEXT("False"), FAssetRegistryTag::TT_Alphabetical));
-		OutTags.Add(FAssetRegistryTag("MLDeformer.NearestNeighborModel.UseDualQuaternions", bUseDualQuaternionDeltas ? TEXT("True") : TEXT("False"), FAssetRegistryTag::TT_Alphabetical));
-		OutTags.Add(FAssetRegistryTag("MLDeformer.NearestNeighborModel.IsNearestNeighborDataValid", bNearestNeighborDataValid ? TEXT("True") : TEXT("False"), FAssetRegistryTag::TT_Alphabetical));
-		OutTags.Add(FAssetRegistryTag("MLDeformer.NearestNeighborModel.IsMorphTargetDataValid", bMorphTargetDataValid ? TEXT("True") : TEXT("False"), FAssetRegistryTag::TT_Alphabetical));
-		OutTags.Add(FAssetRegistryTag("MLDeformer.NearestNeighborModel.UseFileCache", bUseFileCache ? TEXT("True") : TEXT("False"), FAssetRegistryTag::TT_Alphabetical));
-		OutTags.Add(FAssetRegistryTag("MLDeformer.NearestNeighborModel.RecomputDeltas", bRecomputeDeltas ? TEXT("True") : TEXT("False"), FAssetRegistryTag::TT_Alphabetical));
-		OutTags.Add(FAssetRegistryTag("MLDeformer.NearestNeighborModel.RecomputPCA", bRecomputePCA ? TEXT("True") : TEXT("False"), FAssetRegistryTag::TT_Alphabetical));
-	#endif
-
-	OutTags.Add(FAssetRegistryTag("MLDeformer.NearestNeighborModel.UseInputMultipliers", bUseInputMultipliers ? TEXT("True") : TEXT("False"), FAssetRegistryTag::TT_Alphabetical));
-	OutTags.Add(FAssetRegistryTag("MLDeformer.NearestNeighborModel.UseRBF", bUseRBF ? TEXT("True") : TEXT("False"), FAssetRegistryTag::TT_Alphabetical));
-	OutTags.Add(FAssetRegistryTag("MLDeformer.NearestNeighborModel.NumParts", FString::FromInt(ClothPartData.Num()), FAssetRegistryTag::TT_Numerical));
-	OutTags.Add(FAssetRegistryTag("MLDeformer.NearestNeighborModel.DecayFactor", FString::Printf(TEXT("%f"), DecayFactor), FAssetRegistryTag::TT_Numerical));
-	OutTags.Add(FAssetRegistryTag("MLDeformer.NearestNeighborModel.NearestNeighborOffsetWeight", FString::Printf(TEXT("%f"), NearestNeighborOffsetWeight), FAssetRegistryTag::TT_Numerical));
-	OutTags.Add(FAssetRegistryTag("MLDeformer.NearestNeighborModel.RBFSigma", FString::Printf(TEXT("%f"), RBFSigma), FAssetRegistryTag::TT_Numerical));
-	OutTags.Add(FAssetRegistryTag("MLDeformer.NearestNeighborModel.NumInputMultipliers", FString::FromInt(InputMultipliers.Num()), FAssetRegistryTag::TT_Numerical));
+	return "Nearest Neighbor Model";
 }
 
-UMLDeformerInputInfo* UNearestNeighborModel::CreateInputInfo()
+bool UNearestNeighborModel::IsTrained() const
 {
-	UNearestNeighborModelInputInfo* NearestNeighborModelInputInfo = NewObject<UNearestNeighborModelInputInfo>(this);
-#if WITH_EDITORONLY_DATA
-	NearestNeighborModelInputInfo->InitRefBoneRotations(GetSkeletalMesh());
-#endif
-	return NearestNeighborModelInputInfo;
+	return GetOptimizedNetwork() != nullptr;
 }
 
-bool UNearestNeighborModel::IsTrained() const
+FString UNearestNeighborModel::GetDefaultDeformerGraphAssetPath() const
 {
-	return DoesUseOptimizedNetwork() && GetOptimizedNetwork() != nullptr;
+	if (bUseDualQuaternionDeltas)
+	{
+		return TEXT("/NearestNeighborModel/Deformers/DG_DQ_RecomputeNormals.DG_DQ_RecomputeNormals");
+	}
+	else
+	{
+		return TEXT("/DeformerGraph/Deformers/DG_LinearBlendSkin_Morph_Cloth_RecomputeNormals.DG_LinearBlendSkin_Morph_Cloth_RecomputeNormals");
+	}
 }
 
 UMLDeformerModelInstance* UNearestNeighborModel::CreateModelInstance(UMLDeformerComponent* Component)
@@ -109,450 +857,471 @@ UMLDeformerModelInstance* UNearestNeighborModel::CreateModelInstance(UMLDeformer
 void UNearestNeighborModel::Serialize(FArchive& Archive)
 {
 	Super::Serialize(Archive);
+	using UE::NearestNeighborModel::FNearestNeighborModelCustomVersion;
+	if (Archive.IsLoading())
+	{
+		Version = Archive.CustomVer(FNearestNeighborModelCustomVersion::GUID);
+	}
+	Archive.UsingCustomVersion(FNearestNeighborModelCustomVersion::GUID);
 }
 
 void UNearestNeighborModel::PostLoad()
 {
 	Super::PostLoad();
 
-	InitInputInfo();
-
 #if WITH_EDITORONLY_DATA
 	UpdateNetworkInputDim();
 	UpdateNetworkOutputDim();
-	UpdateNetworkSize();
-	UpdateMorphTargetSize();
-#endif
-}
 
-TArray<uint32> ReadTxt(const FString &Path)
-{
-	TArray<FString> StrArr;
-	FFileHelper::LoadFileToStringArray(StrArr, *Path);
-	TArray<uint32> UIntArr; UIntArr.SetNum(StrArr.Num());
-	for (int32 i = 0; i < UIntArr.Num(); i++)
+	UpdateFileCache();
+
+	for (FSection* Section : Sections)
 	{
-		UIntArr[i] = FCString::Atoi(*StrArr[i]);
+		if (Section)
+		{
+			Section->SetModel(this);
+		}
 	}
 
-	return UIntArr;
-}
-
-TArray<uint32> Range(const uint32 Start, const uint32 End)
-{
-	TArray<uint32> Arr; Arr.SetNum(End - Start);
-	for (uint32 i = Start; i < End; i++)
+	if (IsBeforeCustomVersionWasAdded())
 	{
-		Arr[i - Start] = i;
+	PRAGMA_DISABLE_DEPRECATION_WARNINGS 
+		const int32 NumSections = ClothPartData_DEPRECATED.Num();
+		Sections.SetNumZeroed(NumSections);
+		for (int32 Index = 0; Index < NumSections; ++Index)
+		{
+			FSection* Section = OnSectionAdded(Index);
+			check(Section);
+			Section->InitFromClothPartData(ClothPartData_DEPRECATED[Index]);
+		}
+		UpdateForTraining();
+		UpdateForInference();
+	PRAGMA_ENABLE_DEPRECATION_WARNINGS
 	}
-	return Arr;
+#endif
 }
 
-TArray<uint32> AddConstant(const TArray<uint32> &InArr, uint32 Constant)
+int32 UNearestNeighborModel::GetNumSections() const
 {
-	TArray<uint32> OutArr = InArr;
-	for (int32 i = 0; i < InArr.Num(); i++)
-	{
-		OutArr[i] = InArr[i] + Constant;
-	}
-	return OutArr;
+	return Sections.Num();
 }
 
-#if WITH_EDITOR
-int32 UNearestNeighborModel::GetAssetNumNeighbors(int32 PartId) const
+const UNearestNeighborModelSection* UNearestNeighborModel::GetSectionPtr(int32 Index) const
 {
-	return FMath::Min(GetNumNeighborsFromGeometryCache(PartId), GetNumNeighborsFromAnimSequence(PartId));
+	return Sections[Index];
 }
-#endif
 
-TArray<float> UNearestNeighborModel::ClipInputs(const TArray<float>& Inputs) const
+const UNearestNeighborModelSection& UNearestNeighborModel::GetSection(int32 Index) const
 {
-	TArray<float> Result = Inputs;
-	ClipInputs(Result.GetData(), Result.Num());
-	return Result;
+	return *Sections[Index];
 }
 
-void UNearestNeighborModel::ClipInputs(float* InputPtr, int NumInputs) const
+const TArray<int32>& UNearestNeighborModel::GetPCACoeffStarts() const
 {
-	if(InputsMin.Num() >= NumInputs && InputsMax.Num() >= NumInputs)
-	{
-		for(int32 i = 0; i < NumInputs; i++)
-		{
-			float Max = InputsMax[i];
-			float Min = InputsMin[i];
-			if (bUseInputMultipliers && i / 3 < InputMultipliers.Num())
-			{
-				const float Multiplier = InputMultipliers[i / 3][i % 3];
-				Max *= Multiplier;
-				Min *= Multiplier;	
-			}
-
-			if (InputPtr[i] > Max)
-			{
-				InputPtr[i] = Max;
-			}
-			if (InputPtr[i] < Min)
-			{
-				InputPtr[i] = Min;
-			}
-		}
-	}
+	return PCACoeffStarts;
 }
 
 int32 UNearestNeighborModel::GetTotalNumPCACoeffs() const
 {
-	int32 TotalNumPCACoeffs = 0;
-	for (int32 PartId = 0; PartId < GetNumParts(); PartId++)
+	int32 Sum = 0;
+	for(const FSection* Section : Sections)
 	{
-		TotalNumPCACoeffs += GetPCACoeffNum(PartId);
+		Sum += Section->GetNumPCACoeffs();
 	}
-	return TotalNumPCACoeffs;
+	return Sum;
 }
 
 int32 UNearestNeighborModel::GetTotalNumNeighbors() const
 {
-	int32 TotalNumNeighbors = 0;
-	for (int32 PartId = 0; PartId < GetNumParts(); PartId++)
+	int32 Sum = 0;
+	for(const FSection* Section : Sections)
 	{
-		TotalNumNeighbors += GetNumNeighbors(PartId);
+		Sum += Section->GetRuntimeNumNeighbors();
 	}
-	return TotalNumNeighbors;
+	return Sum;
 }
 
-class FMorphTargetBuffersInfo : public FMorphTargetVertexInfoBuffers
+float UNearestNeighborModel::GetDecayFactor() const
 {
-public:
-	uint32 GetMorphDataSize() const { return MorphData.Num() * sizeof(uint32); }
-	void ResetCPUData() { FMorphTargetVertexInfoBuffers::ResetCPUData(); }
-};
+	return DecayFactor;
+}
 
-void UNearestNeighborModel::ResetMorphBuffers()
+float UNearestNeighborModel::GetNearestNeighborOffsetWeight() const
 {
-	if (GetMorphTargetSet())
-	{
-		FMorphTargetBuffersInfo* MorphBuffersInfo = static_cast<FMorphTargetBuffersInfo*>(&GetMorphTargetSet()->MorphBuffers);
-		MorphBuffersInfo->ResetCPUData();
-	}
+	return NearestNeighborOffsetWeight;
 }
 
-#if WITH_EDITORONLY_DATA
-void UNearestNeighborModel::UpdatePCACoeffNums()
+void UNearestNeighborModel::ClipInputs(TArrayView<float> Inputs) const
 {
-	uint32 PCACoeffStart = 0;
-	for(int32 PartId = 0; PartId < GetNumParts(); PartId++)
+	const int32 Num = FMath::Min(Inputs.Num(), FMath::Min(InputsMin.Num(), InputsMax.Num()));
+	for(int32 Index = 0; Index < Num; Index++)
 	{
-		ClothPartData[PartId].PCACoeffStart = PCACoeffStart;
-		PCACoeffStart += ClothPartData[PartId].PCACoeffNum;
+		float Max = InputsMax[Index];
+		float Min = InputsMin[Index];
+		if (bUseInputMultipliers && Index / 3 < InputMultipliers.Num())
+		{
+			const float Multiplier = InputMultipliers[Index / 3][Index % 3];
+			Max *= Multiplier;
+			Min *= Multiplier;
+		}
+		Inputs[Index] = FMath::Clamp(Inputs[Index], Min, Max);
 	}
 }
 
-void UNearestNeighborModel::UpdateNetworkInputDim()
+
+TWeakObjectPtr<const UNearestNeighborOptimizedNetwork> UNearestNeighborModel::GetOptimizedNetwork() const
 {
-	InputDim = 3 * GetBoneIncludeList().Num();
+	return OptimizedNetwork;
 }
 
-void UNearestNeighborModel::UpdateNetworkOutputDim()
+TWeakObjectPtr<UNearestNeighborOptimizedNetwork> UNearestNeighborModel::GetOptimizedNetwork()
 {
-	OutputDim = 0;
-	for (int32 i = 0; i < GetNumParts(); i++)
-	{
-		OutputDim += ClothPartData[i].PCACoeffNum; 
-	}
+	return OptimizedNetwork;
 }
 
-void UNearestNeighborModel::UpdateInputMultipliers()
+bool UNearestNeighborModel::IsReadyForTraining() const
 {
-	if (bUseInputMultipliers)
-	{
-		const int32 NumMultipliers = InputMultipliers.Num();
-		const int32 NumBones = GetBoneIncludeList().Num();
-		InputMultipliers.SetNum(NumBones);
-		for (int32 Index = NumMultipliers; Index < NumBones; Index++)
-		{
-			InputMultipliers[Index] = FVector3f(1.f, 1.f, 1.f);
-		}
-	}
+	return bIsReadyForTraining;
 }
 
-bool UNearestNeighborModel::CheckPCAData(int32 PartId) const
+bool UNearestNeighborModel::IsReadyForInference() const
 {
-	const FClothPartData& Data = ClothPartData[PartId];
-	return Data.VertexMap.Num() > 0 && Data.PCABasis.Num() == Data.VertexMap.Num() * 3 * Data.PCACoeffNum;
+	return bIsReadyForInference;
 }
 
-UE::NearestNeighborModel::EUpdateResult UNearestNeighborModel::UpdateVertexMap(int32 PartId, const FString& VertexMapPath, const FSkelMeshImportedMeshInfo& Info)
+bool UNearestNeighborModel::DoesUseRBF() const
 {
-	using namespace UE::NearestNeighborModel;
-	uint8 ReturnCode = EUpdateResult::SUCCESS;
-	const int32 StartIndex = Info.StartImportedVertex;
-	const int32 NumVertices = Info.NumVertices;
-	bool bIsVertexMapValid = true;
-
-	if (VertexMapPath.IsEmpty())
-	{
-		bIsVertexMapValid = false;
-	}
-	else if (!FPaths::FileExists(VertexMapPath))
-	{
-		bIsVertexMapValid = false;
-		ReturnCode |= EUpdateResult::ERROR;
-		UE_LOG(LogNearestNeighborModel, Error, TEXT("Part %d txt path %s does not exist"), PartId, *VertexMapPath, NumVertices);
-	}
-	else if (GetSkeletalMesh() == nullptr)
-	{
-		ReturnCode |= EUpdateResult::ERROR;
-		UE_LOG(LogNearestNeighborModel, Error, TEXT("SkeletalMesh is None"));
-	}
-	else
-	{
-		const TArray<uint32> PartVertexMap = ReadTxt(ClothPartEditorData[PartId].VertexMapPath);
-		const uint32 MaxVertexIndex = FMath::Max(PartVertexMap);
-		if ((int32)MaxVertexIndex >= GetSkeletalMesh()->GetNumImportedVertices())
-		{
-			ReturnCode |= EUpdateResult::ERROR;
-			UE_LOG(LogNearestNeighborModel, Error, TEXT("Part %d vertex map contains invalid vertex index %d, max vertex index is %d"), PartId, MaxVertexIndex, GetSkeletalMesh()->GetNumImportedVertices() - 1);
-		}
-		else
-		{
-			ClothPartData[PartId].VertexMap = PartVertexMap;
-		}
-		return (EUpdateResult)ReturnCode;
-	}
+	return bUseRBF;
+}
 
-	ClothPartData[PartId].VertexMap = Range(StartIndex, StartIndex + NumVertices);
-	return (EUpdateResult)ReturnCode;
+float UNearestNeighborModel::GetRBFSigma() const
+{
+	return RBFSigma;
 }
 
-UE::NearestNeighborModel::EUpdateResult UNearestNeighborModel::UpdateClothPartData()
+#if WITH_EDITOR
+UNearestNeighborModelSection* UNearestNeighborModel::OnSectionAdded(int32 NewIndex)
 {
-	using namespace UE::NearestNeighborModel;
-	uint8 ReturnCode = EUpdateResult::SUCCESS;
-	if(ClothPartEditorData.Num() == 0)
+	if (!Sections.IsValidIndex(NewIndex))
 	{
-		UE_LOG(LogNearestNeighborModel, Error, TEXT("There should be at least 1 cloth part"));
-		return EUpdateResult::ERROR;
+		return nullptr;
 	}
-	if (!GetSkeletalMesh() || !GetSkeletalMesh()->GetImportedModel() || GetSkeletalMesh()->GetImportedModel()->LODModels[0].ImportedMeshInfos.IsEmpty())
+	UNearestNeighborModelSection* Section = NewObject<UNearestNeighborModelSection>(this);
+	if (!Section)
 	{
-		UE_LOG(LogNearestNeighborModel, Error, TEXT("SkeletalMesh is None or SkeletalMesh has no imported model."));
-		return EUpdateResult::ERROR;
+		return nullptr;
 	}
-	const TArray<FSkelMeshImportedMeshInfo>& SkelMeshInfos = GetSkeletalMesh()->GetImportedModel()->LODModels[0].ImportedMeshInfos;
+	Section->SetModel(this);
+	Sections[NewIndex] = Section;
+	return Section;
+}
 
-	ClothPartData.SetNum(ClothPartEditorData.Num());
-	for(int32 PartId = 0; PartId < GetNumParts(); PartId++)
-	{
-		ClothPartData[PartId].PCACoeffNum = ClothPartEditorData[PartId].PCACoeffNum;
-		int32& MeshIndex = ClothPartEditorData[PartId].MeshIndex;
-		// Set to 0 if MeshIndex is invalid. This could happen if the skeletal mesh is changed.
-		MeshIndex = MeshIndex < SkelMeshInfos.Num() ? MeshIndex : 0;
-		ReturnCode |= UpdateVertexMap(PartId, ClothPartEditorData[PartId].VertexMapPath, SkelMeshInfos[MeshIndex]);
-		ClothPartData[PartId].NumVertices = ClothPartData[PartId].VertexMap.Num();
+UNearestNeighborModelSection& UNearestNeighborModel::GetSection(int32 Index)
+{
+	return *Sections[Index];
+}
 
-		if (!CheckPCAData(PartId))
-		{
-			ClothPartData[PartId].VertexMean.SetNumZeroed(ClothPartData[PartId].NumVertices * 3);
-			ClothPartData[PartId].PCABasis.SetNumZeroed(ClothPartData[PartId].NumVertices * 3 * ClothPartData[PartId].PCACoeffNum);
+FDateTime UNearestNeighborModel::GetNetworkLastWriteTime() const
+{
+	return NetworkLastWriteTime;
+}
 
-			// Init default neighbor data.
-			ClothPartData[PartId].AssetNeighborCoeffs.SetNumZeroed(ClothPartData[PartId].PCACoeffNum);
-			ClothPartData[PartId].AssetNeighborOffsets.SetNumZeroed(ClothPartData[PartId].NumVertices * 3);
-			ClothPartData[PartId].NumNeighbors = 1;
-		}
-	}
-	UpdatePCACoeffNums();
+FString UNearestNeighborModel::GetNetworkLastWriteArchitectureString() const
+{
+	return NetworkLastWriteArchitectureString;
+}
 
-	NearestNeighborData.SetNumZeroed(GetNumParts());
-	UpdateNetworkInputDim();
-	UpdateNetworkOutputDim();
-	bClothPartDataValid = true;
-	return (EUpdateResult)ReturnCode;
+FDateTime UNearestNeighborModel::GetMorphTargetsLastWriteTime() const
+{
+	return MorphTargetsLastWriteTime;
 }
 
-TObjectPtr<UAnimSequence> UNearestNeighborModel::GetNearestNeighborSkeletons(int32 PartId)
-{ 
-	if (PartId < NearestNeighborData.Num())
-	{
-		return NearestNeighborData[PartId].Skeletons; 
-	}
-	return nullptr;
+TArray<FString> UNearestNeighborModel::GetCachedDeltasPaths() const
+{
+	return {FString::Printf(TEXT("%s/deltas.bin"), *GetModelDir())};
 }
 
-const TObjectPtr<UAnimSequence> UNearestNeighborModel::GetNearestNeighborSkeletons(int32 PartId) const
-{ 
-	if (PartId < NearestNeighborData.Num())
+TArray<FString> UNearestNeighborModel::GetCachedPCAPaths() const
+{
+	const int32 NumSections = GetNumSections();
+	TArray<FString> FilePaths; 
+	FilePaths.Reserve(NumSections * 2);
+	for (int32 SectionIndex = 0; SectionIndex < NumSections; ++SectionIndex)
 	{
-		return NearestNeighborData[PartId].Skeletons; 
+		const FString VertexMeanPath = FString::Printf(TEXT("%s/vertex_mean_%d.npy"), *GetModelDir(), SectionIndex);
+		FilePaths.Add(VertexMeanPath);
+		const FString PCABasisPath = FString::Printf(TEXT("%s/pca_basis_%d.npy"), *GetModelDir(), SectionIndex);
+		FilePaths.Add(PCABasisPath);
 	}
-	return nullptr;
+	return FilePaths;
+}
+
+TArray<FString> UNearestNeighborModel::GetCachedNetworkPaths() const
+{
+	return {FString::Printf(TEXT("%s/NearestNeighborModel.ubnne"), *GetModelDir())};
+}
+
+TOptional<FDateTime> UNearestNeighborModel::GetCachedDeltasTimestamp() const
+{
+	return CachedDeltasTimestamp;
+}
+
+TOptional<FDateTime> UNearestNeighborModel::GetCachedPCATimestamp() const
+{
+	return CachedPCATimestamp;
+}
+
+TOptional<FDateTime> UNearestNeighborModel::GetCachedNetworkTimestamp() const
+{
+	return CachedNetworkTimestamp; 
 }
 
-TObjectPtr<UGeometryCache> UNearestNeighborModel::GetNearestNeighborCache(int32 PartId)
-{ 
-	if (PartId < NearestNeighborData.Num())
+FMLDeformerGeomCacheTrainingInputAnim* UNearestNeighborModel::GetNearestNeighborAnim(int32 SectionIndex)
+{
+	if (FMath::IsWithin<int32>(SectionIndex, 0, Sections.Num()))
 	{
-		return NearestNeighborData[PartId].Cache;
+		return Sections[SectionIndex]->GetInputAnim();
 	}
 	return nullptr;
 }
 
-const TObjectPtr<UGeometryCache> UNearestNeighborModel::GetNearestNeighborCache(int32 PartId) const
-{ 
-	if (PartId < NearestNeighborData.Num())
+const FMLDeformerGeomCacheTrainingInputAnim* UNearestNeighborModel::GetNearestNeighborAnim(int32 SectionIndex) const
+{
+	if (FMath::IsWithin<int32>(SectionIndex, 0, Sections.Num()))
 	{
-		return NearestNeighborData[PartId].Cache;
+		return Sections[SectionIndex]->GetInputAnim();
 	}
 	return nullptr;
 }
 
-int32 UNearestNeighborModel::GetNumNeighborsFromGeometryCache(int32 PartId) const
+void UNearestNeighborModel::UpdatePCACoeffStarts()
 {
-	const UGeometryCache* Cache = GetNearestNeighborCache(PartId);
-	if (Cache)
-	{
-		const int32 StartFrame = Cache->GetStartFrame();
-		const int32 EndFrame = Cache->GetEndFrame();
-		return EndFrame - StartFrame + 1;
-	}
-	else
+	uint32 Acc = 0;
+	PCACoeffStarts.Reserve(Sections.Num());
+	for(const FSection* Section : Sections)
 	{
-		return 0;
+		PCACoeffStarts.Add(Acc);
+		Acc += Section->GetNumPCACoeffs();
 	}
 }
 
-int32 UNearestNeighborModel::GetNumNeighborsFromAnimSequence(int32 PartId) const
+void UNearestNeighborModel::UpdateNetworkInputDim()
 {
-	const UAnimSequence* Anim = GetNearestNeighborSkeletons(PartId);
-	if (Anim)
-	{
-		return Anim->GetDataModel()->GetNumberOfKeys();
-	}
-	else
-	{
-		return 0;
-	}
+	InputDim = 3 * GetBoneIncludeList().Num();
 }
 
-void UNearestNeighborModel::UpdateNetworkSize()
+void UNearestNeighborModel::UpdateNetworkOutputDim()
 {
+	OutputDim = GetTotalNumPCACoeffs();
+}
 
-	SavedNetworkSize = 0.0f;
+bool UNearestNeighborModel::IsBeforeCustomVersionWasAdded() const
+{
+	using UE::NearestNeighborModel::FNearestNeighborModelCustomVersion;
+	return Version < FNearestNeighborModelCustomVersion::BeforeCustomVersionWasAdded;
 }
 
-void UNearestNeighborModel::UpdateMorphTargetSize()
+
+UE::NearestNeighborModel::EOpFlag UNearestNeighborModel::CheckHiddenLayerDims()
 {
-	if (GetMorphTargetSet())
+	check(InputDim > 0 && OutputDim > 0);
+	if (HiddenLayerDims.IsEmpty())
 	{
-		const FMorphTargetBuffersInfo* MorphBuffersInfo = static_cast<FMorphTargetBuffersInfo*>(&GetMorphTargetSet()->MorphBuffers);
-		const double Size = MorphBuffersInfo->GetMorphDataSize();
-		MorphDataSize = Size / 1024 / 1024;
+		HiddenLayerDims = UE::NearestNeighborModel::Private::GetDefaultHiddenLayerDims(InputDim, OutputDim);
+		UE_LOG(LogNearestNeighborModel, Error, TEXT("HiddenLayerDims is empty, which will lead to bad results. Default hidden layer dims are added for you. Please double check and train model again."));
+		return EOpFlag::Error;
 	}
-	else
+	for (int32 Index = 0; Index < HiddenLayerDims.Num(); ++Index)
 	{
-		MorphDataSize = 0.0f;
+		if (HiddenLayerDims[Index] <= 0)
+		{
+			UE_LOG(LogNearestNeighborModel, Error, TEXT("HiddenLayerDims[%d] is %d. Non-positive numbers will lead to bad results"), Index, HiddenLayerDims[Index]);
+			return EOpFlag::Error;
+		}
 	}
+	return EOpFlag::Success;
 }
 
-FString UNearestNeighborModel::GetModelDir() const
+void UNearestNeighborModel::UpdateInputMultipliers()
 {
-	if (bUseFileCache)
-	{
-		return FileCacheDirectory;
-	}
-	else
+	if (bUseInputMultipliers)
 	{
-		return FPaths::ProjectIntermediateDir() + TEXT("NearestNeighborModel/");
+		const int32 NumMultipliers = InputMultipliers.Num();
+		const int32 NumBones = GetBoneIncludeList().Num();
+		InputMultipliers.SetNum(NumBones);
+		for (int32 Index = NumMultipliers; Index < NumBones; ++Index)
+		{
+			InputMultipliers[Index] = FVector3f(1.0f, 1.0f, 1.0f);
+		}
 	}
 }
 
-bool UNearestNeighborModel::ShouldUseOptimizedNetwork() const
+void UNearestNeighborModel::UpdateFileCache()
 {
-	return UE::NearestNeighborModel::bNearestNeighborModelUseOptimizedNetwork;
+	UpdateCachedDeltasTimestamp();
+	UpdateCachedPCATimestamp();
+	UpdateCachedNetworkTimestamp();
 }
 
-void UNearestNeighborModel::SetUseOptimizedNetwork(bool bInUseOptimizedNetwork)
+const FString& UNearestNeighborModel::GetFileCacheDirectory() const
 {
-	bUseOptimizedNetwork = bInUseOptimizedNetwork;
+	return FileCacheDirectory;
 }
 
-int32 UNearestNeighborModel::GetMaxPartMeshIndex() const
+void UNearestNeighborModel::SetFileCacheDirectory(const FString& InFileCacheDirectory)
 {
-	int32 MaxIndex = -1;
-	for (const FClothPartEditorData& PartData : ClothPartEditorData)
-	{
-		MaxIndex = FMath::Max(MaxIndex, PartData.MeshIndex);
-	}
-	return MaxIndex;
+	FileCacheDirectory = InFileCacheDirectory;
 }
-#endif
 
-void UNearestNeighborModel::InitInputInfo()
+void UNearestNeighborModel::UpdateCachedDeltasTimestamp()
 {
-	UNearestNeighborModelInputInfo* NearestNeighborModelInputInfo = static_cast<UNearestNeighborModelInputInfo*>(GetInputInfo());
-	check(NearestNeighborModelInputInfo != nullptr);
-#if WITH_EDITORONLY_DATA
-	NearestNeighborModelInputInfo->InitRefBoneRotations(GetSkeletalMesh());
-#endif
+	CachedDeltasTimestamp = UE::NearestNeighborModel::Private::GetLatestTimeStamp(GetCachedDeltasPaths());
 }
 
-void UNearestNeighborModel::SetOptimizedNetwork(UNearestNeighborOptimizedNetwork* InOptimizedNetwork)
+void UNearestNeighborModel::UpdateCachedPCATimestamp()
 {
-	OptimizedNetwork = InOptimizedNetwork;
-	GetReinitModelInstanceDelegate().Broadcast();
+	CachedPCATimestamp = UE::NearestNeighborModel::Private::GetLatestTimeStamp(GetCachedPCAPaths());
+}
+
+void UNearestNeighborModel::UpdateCachedNetworkTimestamp()
+{
+	CachedNetworkTimestamp = UE::NearestNeighborModel::Private::GetLatestTimeStamp(GetCachedNetworkPaths());
 }
 
-bool UNearestNeighborModel::DoesUseOptimizedNetwork() const
+void UNearestNeighborModel::UpdateMorphTargetsLastWriteTime()
 {
-	return bUseOptimizedNetwork;
+	MorphTargetsLastWriteTime = FDateTime::UtcNow().GetTicks();
 }
 
-/**
- * Get derived class default object (CDO)
- * @return A pointer to a CDO of the derived python class. 
- */
-template<class T>
-T* GetDerivedCDO()
+UE::NearestNeighborModel::EOpFlag UNearestNeighborModel::UpdateForTraining()
 {
-	TArray<UClass*> Classes;
-	GetDerivedClasses(T::StaticClass(), Classes);
-	if (Classes.IsEmpty())
+	if (IsReadyForTraining())
 	{
-		return nullptr;
+		return EOpFlag::Success;
+	}
+	using namespace UE::NearestNeighborModel;
+	EOpFlag Result = EOpFlag::Success;
+	UpdatePCACoeffStarts();
+	UpdateNetworkInputDim();
+	if (InputDim == 0)
+	{
+		UE_LOG(LogNearestNeighborModel, Error, TEXT("InputDim is 0. Please add bones to the include list."));
+		return EOpFlag::Error;
+	}
+	UpdateNetworkOutputDim();
+	if (OutputDim == 0)
+	{
+		UE_LOG(LogNearestNeighborModel, Error, TEXT("OutputDim is 0. Please create at least one section and use non-zero PCA coeff counts."));
+		return EOpFlag::Error;
+	}
+	Result |= CheckHiddenLayerDims();
+	if (OpFlag::HasError(Result))
+	{
+		return Result;
 	}
 
-	T* Object = Cast<T>(Classes.Last()->GetDefaultObject());
-	return Object;
+	if (GetNumSections() == 0)
+	{
+		UE_LOG(LogNearestNeighborModel, Error, TEXT("At least one section is required. Please create a section using the '+' button."));
+		return EOpFlag::Error;
+	}
+	for (FSection* Section : Sections)
+	{
+		Section->SetModel(this);
+		Result |= Section->UpdateForTraining();
+	}
+	if (!OpFlag::HasError(Result))
+	{
+		bIsReadyForTraining = true;
+	}
+	return Result;
 }
 
-bool UNearestNeighborModel::LoadOptimizedNetwork(const FString& OnnxPath)
+void UNearestNeighborModel::InvalidateTraining()
 {
-	const FString OnnxFile = FPaths::ConvertRelativePathToFull(OnnxPath);
-	if (FPaths::FileExists(OnnxFile))
+	InvalidateTrainingModelOnly();
+	for (FSection* Section : Sections)
 	{
-		UE_LOG(LogNearestNeighborModel, Display, TEXT("Loading Onnx file '%s'..."), *OnnxFile);
-		UNearestNeighborOptimizedNetwork* Result = NewObject<UNearestNeighborOptimizedNetwork>(this, UNearestNeighborOptimizedNetwork::StaticClass());
-
-		UNearestNeighborOptimizedNetworkLoader* Loader = GetDerivedCDO<UNearestNeighborOptimizedNetworkLoader>();
-		if (Loader == nullptr)
+		if (Section)
 		{
-			return false;
+			Section->InvalidateTraining();
 		}
-		Loader->SetOptimizedNetwork(Result);
-		const bool bSuccess = Loader->LoadOptimizedNetwork(OnnxFile);
-		if (bSuccess)
+	}
+}
+
+void UNearestNeighborModel::InvalidateTrainingModelOnly()
+{
+	bIsReadyForTraining = bIsReadyForInference = false;
+}
+
+UE::NearestNeighborModel::EOpFlag UNearestNeighborModel::UpdateForInference()
+{
+	if (IsReadyForInference())
+	{
+		return EOpFlag::Success;
+	}
+	using namespace UE::NearestNeighborModel;
+	if (!bIsReadyForTraining)
+	{
+		return EOpFlag::Error;
+	}
+
+	EOpFlag Result = EOpFlag::Success;
+	for (FSection* Section : Sections)
+	{
+		Result |= Section->UpdateForInference();
+		if (OpFlag::HasError(Result))
 		{
-			// Clear optimized network to avoid error messages in SetOptimizedNetwork. 
-			OptimizedNetwork = nullptr;
-			SetOptimizedNetwork(Result);
-			return true;
+			return Result;
 		}
 	}
+	UpdateInputMultipliers();
+	bIsReadyForInference = true;
+	return Result;
+}
+
+void UNearestNeighborModel::InvalidateInference()
+{
+	InvalidateInferenceModelOnly();
+	for (FSection* Section : Sections)
+	{
+		Section->InvalidateInference();
+	}
+}
+
+void UNearestNeighborModel::InvalidateInferenceModelOnly()
+{
+	bIsReadyForInference = false;
+}
+
+FString UNearestNeighborModel::GetModelDir() const
+{
+	if (bUseFileCache)
+	{
+		return FileCacheDirectory;
+	}
 	else
 	{
-		UE_LOG(LogNearestNeighborModel, Error, TEXT("Onnx file '%s' does not exist!"), *OnnxFile);
+		return FPaths::ProjectIntermediateDir() + TEXT("NearestNeighborModel/");
 	}
+}
 
-	return false;
+void UNearestNeighborModel::SetOptimizedNetwork(UNearestNeighborOptimizedNetwork* InOptimizedNetwork)
+{
+	OptimizedNetwork = InOptimizedNetwork;
+	GetReinitModelInstanceDelegate().Broadcast();
+}
+
+void UNearestNeighborModel::ClearOptimizedNetwork()
+{
+	if (OptimizedNetwork)
+	{
+		OptimizedNetwork->ConditionalBeginDestroy();
+		OptimizedNetwork = nullptr;
+		NetworkLastWriteTime = 0;
+		NetworkLastWriteArchitectureString = FString();
+		InvalidateInference();
+	}
 }
 
 bool UNearestNeighborModel::LoadOptimizedNetworkFromFile(const FString& Filename)
@@ -571,25 +1340,31 @@ bool UNearestNeighborModel::LoadOptimizedNetworkFromFile(const FString& Filename
 		
 		if (bSuccess)
 		{
-			// Clear optimized network to avoid error messages in SetOptimizedNetwork. 
-			OptimizedNetwork = nullptr;
+			ClearOptimizedNetwork();
 			SetOptimizedNetwork(Result);
+			NetworkLastWriteTime = FDateTime::UtcNow().GetTicks();
+			using UE::NearestNeighborModel::Private::GetArchitectureString;
+			NetworkLastWriteArchitectureString = GetArchitectureString(InputDim, HiddenLayerDims, OutputDim);
+			InvalidateInference();
 			return true;
 		}
+		else
+		{
+			UE_LOG(LogNearestNeighborModel, Error, TEXT("Failed to load Network file '%s'!"), *FullPath);
+			Result->ConditionalBeginDestroy();
+		}
 	}
 	else
 	{
 		UE_LOG(LogNearestNeighborModel, Error, TEXT("Network file '%s' does not exist!"), *FullPath);
 	}
-
 	return false;
 }
+#endif
 
-
-int32 UNearestNeighborModel::GetOptimizedNetworkNumOutputs() const
+int32 UNearestNeighborModel::GetNumNetworkOutputs() const
 {
 	return OptimizedNetwork ? OptimizedNetwork->GetNumOutputs() : 0;
 }
 
-
 #undef LOCTEXT_NAMESPACE
diff --git a/Engine/Plugins/Animation/MLDeformer/NearestNeighborModel/Source/NearestNeighborModel/Private/NearestNeighborModelInputInfo.cpp b/Engine/Plugins/Animation/MLDeformer/NearestNeighborModel/Source/NearestNeighborModel/Private/NearestNeighborModelInputInfo.cpp
--- a/Engine/Plugins/Animation/MLDeformer/NearestNeighborModel/Source/NearestNeighborModel/Private/NearestNeighborModelInputInfo.cpp
+++ b/Engine/Plugins/Animation/MLDeformer/NearestNeighborModel/Source/NearestNeighborModel/Private/NearestNeighborModelInputInfo.cpp
@@ -6,32 +6,6 @@
 
 #include UE_INLINE_GENERATED_CPP_BY_NAME(NearestNeighborModelInputInfo)
 
-
-void UNearestNeighborModelInputInfo::ComputeNetworkInput(const TArray<FQuat>& BoneRotations, float* OutputBuffer, int64 StartIndex) const
-{
-    if (OutputBuffer == nullptr)
-    {
-        return;
-    }
-
-    const int32 NumBones = GetNumBones();
-    check(BoneRotations.Num() == NumBones && RefBoneRotations.Num() == NumBones);
-
-    int64 Index = StartIndex;
-    for(int32 BoneIndex = 0; BoneIndex < NumBones; BoneIndex++)
-    {
-        FQuat RestR = RefBoneRotations[BoneIndex];
-        FQuat LocalR = BoneRotations[BoneIndex];
-
-        const FQuat R =  RestR.Inverse() * LocalR;
-        const float Scalar = R.W >= 0 ? 2.f : -2.f;
-
-        OutputBuffer[Index++] = R.X * Scalar;
-        OutputBuffer[Index++] = R.Y * Scalar;
-        OutputBuffer[Index++] = R.Z * Scalar;
-    }
-}
-
 void UNearestNeighborModelInputInfo::ExtractBoneRotations(USkeletalMeshComponent* SkelMeshComponent, TArray<float>& OutRotations) const
 {
     const int32 NumBones = GetNumBones();
@@ -55,7 +29,7 @@ void UNearestNeighborModelInputInfo::ExtractBoneRotations(USkeletalMeshComponent
         }
     }
 
-    ComputeNetworkInput(BoneRotations, OutRotations.GetData());
+    ComputeNetworkInput(BoneRotations, OutRotations);
 }
 
 int32 UNearestNeighborModelInputInfo::CalcNumNeuralNetInputs() const
@@ -65,7 +39,7 @@ int32 UNearestNeighborModelInputInfo::CalcNumNeuralNetInputs() const
         CurveNames.Num();     // One float per curve.
 }
 
-void UNearestNeighborModelInputInfo::InitRefBoneRotations(USkeletalMesh* SkelMesh)
+void UNearestNeighborModelInputInfo::InitRefBoneRotations(const USkeletalMesh* SkelMesh)
 {
     RefBoneRotations.Reset();
     if (SkelMesh)
@@ -82,3 +56,24 @@ void UNearestNeighborModelInputInfo::InitRefBoneRotations(USkeletalMesh* SkelMes
         }
     }
 }
+
+void UNearestNeighborModelInputInfo::ComputeNetworkInput(TConstArrayView<FQuat> BoneRotations, TArrayView<float> OutputView) const
+{
+    const int32 NumBones = GetNumBones();
+    check(BoneRotations.Num() == NumBones && RefBoneRotations.Num() == NumBones);
+    check(OutputView.Num() >= NumBones * 3);
+
+    int32 Index = 0;
+    for(int32 BoneIndex = 0; BoneIndex < NumBones; BoneIndex++)
+    {
+        FQuat RestR = RefBoneRotations[BoneIndex];
+        FQuat LocalR = BoneRotations[BoneIndex];
+
+        const FQuat R =  RestR.Inverse() * LocalR;
+        const float Scalar = R.W >= 0 ? 2.0f : -2.0f;
+
+        OutputView[Index++] = R.X * Scalar;
+        OutputView[Index++] = R.Y * Scalar;
+        OutputView[Index++] = R.Z * Scalar;
+    }
+}
diff --git a/Engine/Plugins/Animation/MLDeformer/NearestNeighborModel/Source/NearestNeighborModel/Private/NearestNeighborModelInstance.cpp b/Engine/Plugins/Animation/MLDeformer/NearestNeighborModel/Source/NearestNeighborModel/Private/NearestNeighborModelInstance.cpp
--- a/Engine/Plugins/Animation/MLDeformer/NearestNeighborModel/Source/NearestNeighborModel/Private/NearestNeighborModelInstance.cpp
+++ b/Engine/Plugins/Animation/MLDeformer/NearestNeighborModel/Source/NearestNeighborModel/Private/NearestNeighborModelInstance.cpp
@@ -1,75 +1,71 @@
 // Copyright Epic Games, Inc. All Rights Reserved.
 
 #include "NearestNeighborModelInstance.h"
+
+#include "Algo/MinElement.h"
+#include "Components/ExternalMorphSet.h"
+#include "Components/SkeletalMeshComponent.h"
+#include "Engine/SkeletalMesh.h"
 #include "NearestNeighborModel.h"
 #include "NearestNeighborModelInputInfo.h"
 #include "NearestNeighborOptimizedNetwork.h"
-#include "MLDeformerAsset.h"
-#include "Components/ExternalMorphSet.h"
-#include "Engine/SkeletalMesh.h"
-#include "Components/SkeletalMeshComponent.h"
 
 #include UE_INLINE_GENERATED_CPP_BY_NAME(NearestNeighborModelInstance)
 CSV_DECLARE_CATEGORY_MODULE_EXTERN(MLDEFORMERFRAMEWORK_API, MLDeformer);
 
-using namespace UE::MLDeformer;
-using namespace UE::NearestNeighborModel;
-
-int64 UNearestNeighborModelInstance::SetBoneTransforms(float* OutputBuffer, int64 OutputBufferSize, int64 StartIndex)
+void UNearestNeighborModelInstance::Init(USkeletalMeshComponent* SkelMeshComponent)
 {
-	// Get the transforms for the bones we used during training.
-	// These are in the space relative to their parent.
-	UpdateBoneTransforms();
-
-    const UNearestNeighborModelInputInfo* InputInfo = static_cast<UNearestNeighborModelInputInfo*>(Model->GetInputInfo());
-    const int32 NumBones = InputInfo->GetNumBones();
-
-    // Extract BoneRotations from TrainingBoneTransforms (@todo: optimize)
-    TArray<FQuat> BoneRotations; BoneRotations.AddZeroed(NumBones);
-    int64 EndIndex = StartIndex + NumBones * 3;
-	check(EndIndex <= OutputBufferSize); // Make sure we don't write past the OutputBuffer.
-    
-	for (int32 BoneIndex = 0; BoneIndex < NumBones; ++BoneIndex)
+	UMLDeformerMorphModelInstance::Init(SkelMeshComponent);
+	InitInstanceData();
+	const UNearestNeighborModel* const NearestNeighborModel = GetCastModel();
+	if (!NearestNeighborModel)
 	{
-        BoneRotations[BoneIndex] = TrainingBoneTransforms[BoneIndex].GetRotation();
+		return;
 	}
+	if (NearestNeighborModel->GetOptimizedNetwork().IsValid())
+	{
+		InitOptimizedNetworkInstance();
+		check(OptimizedNetworkInstance);
+	}
+}
 
-    // Write the transforms into the output buffer.
-    InputInfo->ComputeNetworkInput(BoneRotations, OutputBuffer, StartIndex);
-
-	return EndIndex;
+void UNearestNeighborModelInstance::Execute(float ModelWeight)
+{
+	CSV_SCOPED_TIMING_STAT(MLDeformer, NearestNeighborExecute);
+	check(OptimizedNetworkInstance);
+	OptimizedNetworkInstance->Run();
 }
 
 bool UNearestNeighborModelInstance::SetupInputs()
 {
-	const UNearestNeighborModel* NearestNeighborModel = Cast<UNearestNeighborModel>(Model);
-	if (NearestNeighborModel == nullptr ||
-		SkeletalMeshComponent == nullptr ||
-		SkeletalMeshComponent->GetSkeletalMeshAsset() == nullptr ||
-		!bIsCompatible)
+	const UNearestNeighborModel* const NearestNeighborModel = GetCastModel();
+	if (!NearestNeighborModel || !NearestNeighborModel->GetInputInfo())
 	{
 		return false;
 	}
-	if (NearestNeighborModel->DoesUseOptimizedNetwork())
+	if (!NearestNeighborModel->IsReadyForInference())
 	{
-		const UNearestNeighborOptimizedNetwork* OptimizedNetwork = NearestNeighborModel->GetOptimizedNetwork();
-		if (!OptimizedNetwork || !OptimizedNetworkInstance)
-		{
-			return false;
-		}
+		return false;
 	}
-	else
+	if (SkeletalMeshComponent == nullptr ||
+		SkeletalMeshComponent->GetSkeletalMeshAsset() == nullptr ||
+		!bIsCompatible)
+	{
+		return false;
+	}
+	UConstNetworkPtr OptimizedNetwork = NearestNeighborModel->GetOptimizedNetwork();
+	if (!OptimizedNetwork.IsValid() || !OptimizedNetworkInstance)
 	{
 		return false;
 	}
 
-	int32 NumNeuralNetInputs = 0;
-	float* InputDataPointer = nullptr;
-	GetInputDataPointer(InputDataPointer, NumNeuralNetInputs);
-	if (!InputDataPointer)
+	TOptional<TArrayView<float>> OptionalInputView = GetInputView();
+	if (!OptionalInputView.IsSet())
 	{
 		return false;
 	}
+	TArrayView<float> InputView = OptionalInputView.GetValue();
+	const int32 NumNeuralNetInputs = InputView.Num();
 
 	// If the neural network expects a different number of inputs, do nothing.
 	const int32 NumFloatsPerBone = NearestNeighborModel->GetNumFloatsPerBone();
@@ -80,47 +76,11 @@ bool UNearestNeighborModelInstance::SetupInputs()
 		return false;
 	}
 
-	const int64 NumFloatsWritten = SetNeuralNetworkInputValues(InputDataPointer, NumNeuralNetInputs);
+	const int64 NumFloatsWritten = SetNeuralNetworkInputValues(InputView.GetData(), NumNeuralNetInputs);
 	check(NumFloatsWritten == NumNeuralNetInputs);
-	if (InputDataPointer)
-	{
-		// Clip network inputs based on training set min and max values.
-		NearestNeighborModel->ClipInputs(InputDataPointer, NumNeuralNetInputs);
-		return true;
-	}
-    return false;
-}
-
-void UNearestNeighborModelInstance::Init(USkeletalMeshComponent* SkelMeshComponent)
-{
-	UMLDeformerMorphModelInstance::Init(SkelMeshComponent);
-	InitPreviousWeights();
-	const UNearestNeighborModel* NearestNeighborModel = Cast<UNearestNeighborModel>(Model);
-	check(NearestNeighborModel);
-	if (NearestNeighborModel->DoesUseOptimizedNetwork())
-	{
-		InitOptimizedNetworkInstance();
-	}
-}
-
-void UNearestNeighborModelInstance::InitOptimizedNetworkInstance()
-{
-	const UNearestNeighborModel* NearestNeighborModel = Cast<UNearestNeighborModel>(Model);
-	check(NearestNeighborModel);
-	UNearestNeighborOptimizedNetwork* Network = NearestNeighborModel->GetOptimizedNetwork();
-	OptimizedNetworkInstance = Network ? Network->CreateInstance() : nullptr;
-}
-
-void UNearestNeighborModelInstance::Execute(float ModelWeight)
-{
-	CSV_SCOPED_TIMING_STAT(MLDeformer, NearestNeighborExecute);
-	const UNearestNeighborModel* NearestNeighborModel = Cast<UNearestNeighborModel>(Model);
-	check(NearestNeighborModel);
-	if (NearestNeighborModel->DoesUseOptimizedNetwork())
-	{
-		UNearestNeighborOptimizedNetwork* OptimizedNetwork = NearestNeighborModel->GetOptimizedNetwork();
-		OptimizedNetworkInstance->Run();
-	}
+	// Clip network inputs based on training set min and max values.
+	NearestNeighborModel->ClipInputs(InputView);
+	return true;
 }
 
 void UNearestNeighborModelInstance::Tick(float DeltaTime, float ModelWeight)
@@ -146,24 +106,144 @@ void UNearestNeighborModelInstance::Tick(float DeltaTime, float ModelWeight)
 	}
 }
 
-TArray<float> ExtractArray(const TArray<float>& InArr, int32 Start, int32 End)
+int64 UNearestNeighborModelInstance::SetBoneTransforms(float* OutputBuffer, int64 OutputBufferSize, int64 StartIndex)
+{
+	// Get the transforms for the bones we used during training.
+	// These are in the space relative to their parent.
+	UpdateBoneTransforms();
+
+	check(Model);
+    const UNearestNeighborModelInputInfo* InputInfo = static_cast<UNearestNeighborModelInputInfo*>(Model->GetInputInfo());
+	check(InputInfo);
+    const int32 NumBones = InputInfo->GetNumBones();
+
+    // Extract BoneRotations from TrainingBoneTransforms (@todo: optimize)
+    TArray<FQuat> BoneRotations; 
+	BoneRotations.SetNumUninitialized(NumBones);
+	const UNearestNeighborModel* const NearestNeighborModel = GetCastModel();
+	check(NearestNeighborModel);
+	const int32 NumFloatsPerBone = NearestNeighborModel->GetNumFloatsPerBone();
+    const int64 EndIndex = StartIndex + NumBones * NumFloatsPerBone;
+	check(EndIndex <= OutputBufferSize); // Make sure we don't write past the OutputBuffer.
+    
+	for (int32 BoneIndex = 0; BoneIndex < NumBones; ++BoneIndex)
+	{
+        BoneRotations[BoneIndex] = TrainingBoneTransforms[BoneIndex].GetRotation();
+	}
+
+    // Write the transforms into the output buffer.
+    InputInfo->ComputeNetworkInput(BoneRotations, TArrayView<float>(OutputBuffer + StartIndex, NumBones * NumFloatsPerBone));
+
+	return EndIndex;
+}
+
+FString UNearestNeighborModelInstance::CheckCompatibility(USkeletalMeshComponent* InSkelMeshComponent, bool LogIssues)
+{
+	ErrorText = FString();
+
+	// If we're not compatible, generate a compatibility string.
+	USkeletalMesh* const SkelMesh = InSkelMeshComponent ? InSkelMeshComponent->GetSkeletalMeshAsset() : nullptr;
+	UMLDeformerInputInfo* const InputInfo = Model->GetInputInfo();
+	if (SkelMesh && !InputInfo->IsCompatible(SkelMesh) && Model->GetDeformerAsset())
+	{
+		ErrorText += InputInfo->GenerateCompatibilityErrorString(SkelMesh);
+		ErrorText += "\n";
+		if (LogIssues)
+		{
+			UE_LOG(LogNearestNeighborModel, Error, TEXT("ML Deformer '%s' isn't compatible with Skeletal Mesh '%s'.\nReason(s):\n%s"), 
+				*Model->GetDeformerAsset()->GetName(), 
+				*SkelMesh->GetName(), 
+				*ErrorText);
+		}
+	}
+
+	const UNearestNeighborModel* const NearestNeighborModel = GetCastModel();
+	if (!NearestNeighborModel)
+	{
+		return FString();
+	}
+	// Verify the number of inputs versus the expected number of inputs.
+	const int32 NumFloatsPerBone = NearestNeighborModel->GetNumFloatsPerBone();
+	const int32 NumFloatsPerCurve = NearestNeighborModel->GetNumFloatsPerCurve();
+	const int64 NumDeformerAssetInputs = static_cast<int64>(Model->GetInputInfo()->CalcNumNeuralNetInputs(NumFloatsPerBone, NumFloatsPerCurve));
+	int64 NumNeuralNetInputs = -1;
+	UConstNetworkPtr OptimizedNetwork = NearestNeighborModel->GetOptimizedNetwork();
+	if (OptimizedNetwork.IsValid())
+	{
+		NumNeuralNetInputs = OptimizedNetwork->GetNumInputs();
+	}
+
+	// Only check compatibility after the network is loaded. 
+	if (NumNeuralNetInputs >= 0 && NumNeuralNetInputs != NumDeformerAssetInputs) 
+	{
+		const FString InputErrorString = "The number of network inputs doesn't match the asset. Please retrain the asset."; 
+		ErrorText += InputErrorString + "\n";
+		if (LogIssues && Model->GetDeformerAsset())
+		{
+			UE_LOG(LogNearestNeighborModel, Error, TEXT("Deformer '%s': %s"), *(Model->GetDeformerAsset()->GetName()), *InputErrorString);
+		}
+	}
+
+	return ErrorText;
+}
+
+void UNearestNeighborModelInstance::InitInstanceData(int32 NumMorphWeights)
 {
-	int32 Num = End - Start;
-	TArray<float> Result; Result.SetNum(Num);
-	for (int32 i = 0; i < Num; i++)
+	if (NumMorphWeights == INDEX_NONE)
 	{
-		Result[i] = InArr[Start + i];
+		const UNearestNeighborModel* const NearestNeighborModel = GetCastModel();
+		if (!NearestNeighborModel)
+		{
+			return;
+		}
+		NumMorphWeights = 1 + NearestNeighborModel->GetTotalNumPCACoeffs() + NearestNeighborModel->GetTotalNumNeighbors();
 	}
-	return Result;
+	PreviousWeights.SetNumZeroed(NumMorphWeights);
+	DistanceBuffer.SetNumZeroed(NumMorphWeights);
 }
 
-constexpr float FastExpMinusOne(float X)
+UNearestNeighborModel* UNearestNeighborModelInstance::GetCastModel() const
 {
-	return X + X * X / 2.0f + X * X * X / 6.0f;
+	return Cast<UNearestNeighborModel>(Model);
+}
+
+#if WITH_EDITOR
+TArray<uint32> UNearestNeighborModelInstance::GetNearestNeighborIds() const
+{
+	return NearestNeighborIds;
+}
+#endif
+
+void UNearestNeighborModelInstance::InitOptimizedNetworkInstance()
+{
+	const UNearestNeighborModel* const NearestNeighborModel = GetCastModel();
+	if (!NearestNeighborModel)
+	{
+		return;
+	}
+	UConstNetworkPtr Network = NearestNeighborModel->GetOptimizedNetwork();
+	OptimizedNetworkInstance = Network.IsValid() ? Network->CreateInstance(this) : nullptr;
 }
 
 namespace UE::NearestNeighborModel::Private
 {
+	constexpr float FastExpMinusOne(float X)
+	{
+		return X + X * X / 2.0f + X * X * X / 6.0f;
+	}
+
+	float GetDecayCoeff(float DeltaTime, float DecayFactor)
+	{
+		/** DecayCoeff = (e ^ {(DecayFactor - 1) * DeltaTime + Delta0} - 1) / (e ^ {Delta0} - 1)
+		 * DecayFactor is a user defined value. DecayCoeff is the actual coefficient used to compute decay. 
+		 * DecayCoeff is an exponential decay of DeltaTime.
+		 * When DecayFactor = 1, DecayCoeff = 1; when DecayFactor = 0 && DeltaTime = Delta0, DecayCoeff = 0
+		 */	
+		constexpr float Delta0 = 1.0f / 30.f;
+		const float Delta = (DecayFactor - 1.0f) * DeltaTime;
+		return FastExpMinusOne(FMath::Max(Delta + Delta0, 0.0f)) / FastExpMinusOne(Delta0);
+	}
+
 	float ComputeDistanceSquared(TConstArrayView<float> V1, TConstArrayView<float> V2)
 	{
 		check(V1.Num() == V2.Num());
@@ -176,25 +256,61 @@ namespace UE::NearestNeighborModel::Private
 		return DistanceSquared;
 	}
 
-	void UpdateRBFWeights(TArrayView<float> OutMorphWeights, int32& OutNearestNeighborId, TArrayView<float> OutPreviousWeights, TConstArrayView<float> CoeffsView, const TConstArrayView<float> NeighborCoeffsView, float Sigma, float OffsetWeight, float DecayCoeff)
+	void ComputeNeighborDistances(TArrayView<float> OutDistances, const UNearestNeighborModel::FSection& Section, TConstArrayView<float> PCACoeffs)
+	{
+		const int32 NumCoeffs = Section.GetNumPCACoeffs();
+		const int32 NumNeighbors = Section.GetRuntimeNumNeighbors();
+		check(OutDistances.Num() == NumNeighbors);
+		TConstArrayView<float> NeighborCoeffs = Section.GetNeighborCoeffs();
+		for (int32 NeighborId = 0; NeighborId < NumNeighbors; ++NeighborId)
+		{
+			TConstArrayView<float> Coeffs(NeighborCoeffs.GetData() + NeighborId * NumCoeffs, NumCoeffs);
+			OutDistances[NeighborId] = ComputeDistanceSquared(Coeffs, PCACoeffs);
+		}
+	}
+
+	void UpdateWeightWithDecay(float& OutDecayedWeight, float& InOutPrevWeight, float NewWeight, float DecayCoeff)
+	{
+		const float NewW = (1.0f - DecayCoeff) * NewWeight + DecayCoeff * InOutPrevWeight;
+		OutDecayedWeight = NewW;
+		InOutPrevWeight = NewW;
+	}
+
+	void UpdateNearestNeighborWeights(TArrayView<float> OutMorphWeights, TArrayView<float> InOutPrevWeights, TConstArrayView<float> SquaredDistances, float OffsetWeight, float DecayCoeff)
 	{
-		const int32 NumPCACoeffs = CoeffsView.Num();
-		check(OutMorphWeights.Num() == NumPCACoeffs);
-		const int32 NumNeighbors = NeighborCoeffsView.Num() / NumPCACoeffs;
-		check(NeighborCoeffsView.Num() == NumPCACoeffs * NumNeighbors);
-		TArray<float> Weights;
-		Weights.SetNumUninitialized(NumNeighbors);
+		const int32 NumNeighbors = OutMorphWeights.Num();
+		check(InOutPrevWeights.Num() == NumNeighbors);
+		check(SquaredDistances.Num() == NumNeighbors);
+		if (NumNeighbors <= 0)
+		{
+			return;
+		}
+
+		const float* MinD2 = Algo::MinElement(SquaredDistances);
+		const int32 MinIndex = MinD2 - SquaredDistances.GetData();
+
 		for (int32 Index = 0; Index < NumNeighbors; ++Index)
 		{
-			TConstArrayView<float> NeighborCoeffs(NeighborCoeffsView.GetData() + Index * NumPCACoeffs, NumPCACoeffs);
-			Weights[Index] = ComputeDistanceSquared(CoeffsView, NeighborCoeffs);
+			const float W = Index == MinIndex ? OffsetWeight : 0.0f;
+			UpdateWeightWithDecay(OutMorphWeights[Index], InOutPrevWeights[Index], W, DecayCoeff);
 		}
-		const float MinD2 = FMath::Min(Weights, &OutNearestNeighborId);
+	}
+
+	// SquaredDistances buffer will be overwritten after this function.
+	void UpdateRBFWeights(TArrayView<float> OutMorphWeights, TArrayView<float> InOutPrevWeights, TArrayView<float> SquaredDistances, float Sigma, float OffsetWeight, float DecayCoeff)
+	{
+		const int32 NumNeighbors = OutMorphWeights.Num();
+		check(InOutPrevWeights.Num() == NumNeighbors);
+		check(SquaredDistances.Num() == NumNeighbors);
+		check(NumNeighbors > 0);
+
+		const float MinD2 = *Algo::MinElement(SquaredDistances);
 		float SumWeights = 0.0f;
 		const float Sigma2 = FMath::Max(Sigma * Sigma, 1e-6f);
+		TArrayView<float> Weights = SquaredDistances;
 		for (int32 Index = 0; Index < NumNeighbors; ++Index)
 		{
-			float Weight = (Weights[Index] - MinD2) / Sigma2;
+			float Weight = (SquaredDistances[Index] - MinD2) / Sigma2;
 			constexpr float CutOff = 3.0f;
 			if (Weight < CutOff)
 			{
@@ -207,277 +323,170 @@ namespace UE::NearestNeighborModel::Private
 				Weights[Index] = 0.0f;
 			}
 		}
+
 		SumWeights = FMath::Max(SumWeights, 1e-6f);
 		for (int32 Index = 0; Index < NumNeighbors; ++Index)
 		{
 			const float W = Weights[Index] / SumWeights * OffsetWeight;
-			const float PreviousW = OutPreviousWeights[Index];
-			const float NewW = (1.0f - DecayCoeff) * W + DecayCoeff * PreviousW;
-			OutMorphWeights[Index] = NewW;
-			OutPreviousWeights[Index] = NewW;
+			UpdateWeightWithDecay(OutMorphWeights[Index], InOutPrevWeights[Index], W, DecayCoeff);
 		}
 	}
-};
+}
 
 void UNearestNeighborModelInstance::RunNearestNeighborModel(float DeltaTime, float ModelWeight)
 {
-	UNearestNeighborModel* NearestNeighborModel = Cast<UNearestNeighborModel>(Model);
-	if (NearestNeighborModel == nullptr)
+	const UNearestNeighborModel* const NearestNeighborModel = GetCastModel();
+	if (!NearestNeighborModel)
 	{
 		return;
 	}
-	
-	/** DecayCoeff = (e ^ {(DecayFactor - 1) * DeltaTime + Delta0} - 1) / (e ^ {Delta0} - 1)
-	 * DecayFactor is a user defined value. DecayCoeff is the actual coefficient used to compute decay. 
-	 * DecayCoeff is an exponential decay of DeltaTime.
-	 * When DecayFactor = 1, DecayCoeff = 1; when DecayFactor = 0 && DeltaTime = Delta0, DecayCoeff = 0
-	 */	
-	constexpr float Delta0 = 1.0f / 30;
-	const float Delta = (NearestNeighborModel->GetDecayFactor() - 1) * DeltaTime;
-	const float DecayCoeff = FastExpMinusOne(FMath::Max(Delta + Delta0, 0)) / FastExpMinusOne(Delta0);
+	using UE::NearestNeighborModel::Private::GetDecayCoeff;
+	const float DecayCoeff = GetDecayCoeff(DeltaTime, NearestNeighborModel->GetDecayFactor());
 
 	// Grab the weight data for this morph set.
 	// This could potentially fail if we are applying this deformer to the wrong skeletal mesh component.
 	const int LOD = 0;	// For now we only support LOD 0, as we can't setup an ML Deformer per LOD yet.
-	FExternalMorphSetWeights* WeightData = FindWeightData(LOD);
+	FExternalMorphSetWeights* const WeightData = FindWeightData(LOD);
 	if (WeightData == nullptr)
 	{
 		return;
 	}
 
-#if WITH_EDITORONLY_DATA
-	NearestNeighborIds.SetNum(NearestNeighborModel->GetNumParts());
+#if WITH_EDITOR
+	NearestNeighborIds.SetNum(NearestNeighborModel->GetNumSections());
 #endif
 
-	float* OutputTensorData = nullptr;
-	int32 NumNetworkWeights = 0;
-	GetOutputDataPointer(OutputTensorData, NumNetworkWeights);
+	TOptional<TArrayView<float>> OptionalOutputView = GetOutputView();
+	if (!OptionalOutputView.IsSet())
+	{
+		WeightData->ZeroWeights();
+		return;
+	}
+	TArrayView<float> OutputView = OptionalOutputView.GetValue();
+
+	const int32 NumNetworkWeights = OutputView.Num();
+	const int32 NumMorphTargets = WeightData->Weights.Num();
+	const int32 TotalNumNeighbors = NearestNeighborModel->GetTotalNumNeighbors();
 
-	if (OutputTensorData)
+	if (NumMorphTargets == NumNetworkWeights + 1 + TotalNumNeighbors)
 	{
-		const int32 NumMorphTargets = WeightData->Weights.Num();
-		const int32 TotalNumNeighbors = NearestNeighborModel->GetTotalNumNeighbors();
-		if (NumMorphTargets >= NumNetworkWeights + 1 + TotalNumNeighbors)
+		InitInstanceData(NumMorphTargets);
+
+		const float OffsetWeight = NearestNeighborModel->GetNearestNeighborOffsetWeight();
+		WeightData->Weights[0] = ModelWeight;
+
+		// Update all generated morph target weights with the values calculated by our neural network.
+		for (int32 MorphIndex = 0; MorphIndex < NumNetworkWeights; ++MorphIndex)
 		{
-			WeightData->Weights[0] = 1.0f * ModelWeight;
+			const float W = OutputView[MorphIndex] * ModelWeight;
+			using UE::NearestNeighborModel::Private::UpdateWeightWithDecay;
+			UpdateWeightWithDecay(WeightData->Weights[MorphIndex + 1], PreviousWeights[MorphIndex + 1], W, DecayCoeff);
+		}
 
-			// Update all generated morph target weights with the values calculated by our neural network.
-			for (int32 MorphIndex = 0; MorphIndex < NumNetworkWeights; ++MorphIndex)
+		int32 NeighborOffset = NumNetworkWeights + 1;
+		for (int32 SectionIndex = 0; SectionIndex < NearestNeighborModel->GetNumSections(); SectionIndex++)
+		{
+			const UNearestNeighborModel::FSection& Section = NearestNeighborModel->GetSection(SectionIndex);
+			if (!Section.IsReadyForInference())
 			{
-				const float W = OutputTensorData[MorphIndex] * ModelWeight;
-				UpdateWeightWithDecay(WeightData->Weights, MorphIndex + 1, W, DecayCoeff);
+				continue;
 			}
 
-			int32 NeighborOffset = NumNetworkWeights + 1;
-			for (int32 PartId = 0; PartId < NearestNeighborModel->GetNumParts(); PartId++)
+			TConstArrayView<int32> PCACoeffStarts = NearestNeighborModel->GetPCACoeffStarts();
+			if (!PCACoeffStarts.IsValidIndex(SectionIndex))
 			{
-				if (NearestNeighborModel->bUseRBF)
-				{
-					const int32 CoeffStart = NearestNeighborModel->GetPCACoeffStart(PartId);
-					const int32 NumPCACoeffs = NearestNeighborModel->GetPCACoeffNum(PartId);
-					TConstArrayView<float> CoeffsView(OutputTensorData + CoeffStart, NumPCACoeffs);
-					check(CoeffStart + NumPCACoeffs <= NumNetworkWeights);
-
-					const TArray<float>& NeighborCoeffs = NearestNeighborModel->NeighborCoeffs(PartId);
-					const int32 NumNeighbors = NearestNeighborModel->GetNumNeighbors(PartId);
-					if(!ensureMsgf(NeighborCoeffs.Num() == NumPCACoeffs * NumNeighbors, TEXT("Nearest Neigbbor Data is out-of-date, please hit Update button")))
-					{
-						continue;
-					}
-
-					TArrayView<float> MorphWeightsView(WeightData->Weights.GetData() + NeighborOffset, NumPCACoeffs);
-					TArrayView<float> PreviousWeightsView(PreviousWeights.GetData() + NeighborOffset, NumPCACoeffs);
-					const float OffsetWeight = ModelWeight * NearestNeighborModel->GetNearestNeighborOffsetWeight();
-					int32 NearestNeighborId = INDEX_NONE;
-					Private::UpdateRBFWeights(MorphWeightsView, NearestNeighborId, PreviousWeightsView, CoeffsView, NeighborCoeffs, NearestNeighborModel->RBFSigma, OffsetWeight, DecayCoeff);
-					#if WITH_EDITORONLY_DATA
-						NearestNeighborIds[PartId] = NearestNeighborId;
-					#endif
-					NeighborOffset += NumNeighbors;
-				}
-				else
-				{
-					const int32 NearestNeighborId = FindNearestNeighbor(OutputTensorData, PartId);
-						#if WITH_EDITORONLY_DATA
-					NearestNeighborIds[PartId] = NearestNeighborId;
-						#endif
-					
-					const int32 NumNeighbors = NearestNeighborModel->GetNumNeighbors(PartId);
-					for (int32 NeighborId = 0; NeighborId < NumNeighbors; NeighborId++)
-					{
-						const float W = NeighborId == NearestNeighborId ? ModelWeight * NearestNeighborModel->GetNearestNeighborOffsetWeight() : 0;
-						const int32 Index = NeighborOffset + NeighborId;
-						if (Index < NumMorphTargets)
-						{
-							UpdateWeightWithDecay(WeightData->Weights, Index, W, DecayCoeff);
-						}
-					}
-					NeighborOffset += NumNeighbors;
-				}
+				continue;
 			}
-		}
-	}
-	else
-	{
-		WeightData->ZeroWeights();
-	}
-}
+			const int32 PCACoeffStart = PCACoeffStarts[SectionIndex];
+			const int32 NumCoeffs = Section.GetNumPCACoeffs();
+			TConstArrayView<float> PCACoeffs(OutputView.GetData() + PCACoeffStart, NumCoeffs);
 
-void UNearestNeighborModelInstance::UpdateWeightWithDecay(TArray<float>& MorphWeights, int32 Index, float W, float DecayCoeff)
-{
-	const float PreviousW = PreviousWeights[Index];
-	const float NewW = (1 - DecayCoeff) * W + DecayCoeff * PreviousW;
-	MorphWeights[Index] = NewW;
-	PreviousWeights[Index] = NewW; 
-}
+			const int32 NumNeighbors = Section.GetRuntimeNumNeighbors();
 
-int32 UNearestNeighborModelInstance::FindNearestNeighbor(const float* PCAData, int32 PartId)
-{
-	const UNearestNeighborModel* NearestNeighborModel = Cast<UNearestNeighborModel>(Model);
-	check(NearestNeighborModel);
-	const TArray<float> &NeighborCoeffs = NearestNeighborModel->NeighborCoeffs(PartId);
-	float MinD2 = std::numeric_limits<float>::max();
-	int32 MinId = -1;
-	for (int32 NeighborId = 0; NeighborId < NearestNeighborModel->GetNumNeighbors(PartId); NeighborId++)
-	{
-		float D2 = 0;
-		const int32 NumCoeffs = NearestNeighborModel->GetPCACoeffNum(PartId);
-		const int32 CoeffStart = NearestNeighborModel->GetPCACoeffStart(PartId);
-		const int32 NumNeighbors = NearestNeighborModel->GetNumNeighbors(PartId);
-		for (int32 CoeffId = 0; CoeffId < NumCoeffs; CoeffId++)
-		{
-			const float Coeff = PCAData[CoeffStart + CoeffId];
-			const float NeighborCoeff = NeighborCoeffs[NeighborId * NumCoeffs + CoeffId];
-			const float D = Coeff - NeighborCoeff;
-			D2 += D * D;
-		}
-		if (D2 < MinD2)
-		{
-			MinD2 = D2;
-			MinId = NeighborId;
-		}
-	}
-	return MinId;
-}
+			TArrayView<float> SquaredDistances(DistanceBuffer.GetData() + NeighborOffset, NumNeighbors);
+			using UE::NearestNeighborModel::Private::ComputeNeighborDistances;
+			ComputeNeighborDistances(SquaredDistances, Section, PCACoeffs);
 
-void UNearestNeighborModelInstance::InitPreviousWeights()
-{
-	const UNearestNeighborModel* NearestNeighborModel = Cast<UNearestNeighborModel>(Model);
-	if (NearestNeighborModel && NearestNeighborModel->GetMorphTargetSet())
-	{
-		const int32 NumWeights = NearestNeighborModel->GetMorphTargetSet()->MorphBuffers.GetNumMorphs();
-		PreviousWeights.SetNumZeroed(NumWeights);
-	}
-}
+			TArrayView<float> SectionMorphWeights(WeightData->Weights.GetData() + NeighborOffset, NumNeighbors);
+			TArrayView<float> SectionPreviousWeights(PreviousWeights.GetData() + NeighborOffset, NumNeighbors);
+			if (NearestNeighborModel->DoesUseRBF())
+			{
+				const float Sigma = NearestNeighborModel->GetRBFSigma();
+				using UE::NearestNeighborModel::Private::UpdateRBFWeights;
+				UpdateRBFWeights(SectionMorphWeights, SectionPreviousWeights, SquaredDistances, Sigma, OffsetWeight, DecayCoeff);
+			}
+			else
+			{
+				using UE::NearestNeighborModel::Private::UpdateNearestNeighborWeights;
+				UpdateNearestNeighborWeights(SectionMorphWeights, SectionPreviousWeights, SquaredDistances, OffsetWeight, DecayCoeff);
+			}
 
-void UNearestNeighborModelInstance::GetInputDataPointer(float*& OutInputData, int32& OutNumInputFloats)
-{
-	const UNearestNeighborModel* NearestNeighborModel = Cast<UNearestNeighborModel>(Model);
-	check(NearestNeighborModel);
-	if (NearestNeighborModel->DoesUseOptimizedNetwork())
-	{
-		const UNearestNeighborOptimizedNetwork* OptimizedNetwork = NearestNeighborModel->GetOptimizedNetwork();
-		if (OptimizedNetwork && OptimizedNetworkInstance)
-		{
-			OutInputData = OptimizedNetworkInstance->GetInputs().GetData();
-			OutNumInputFloats = OptimizedNetwork->GetNumInputs();
-			return;
+#if WITH_EDITOR
+			const int32 NeighborId = NumNeighbors > 0 ? (Algo::MinElement(SquaredDistances) - SquaredDistances.GetData()) : INDEX_NONE;
+			NearestNeighborIds[SectionIndex] = NeighborId;
+#endif // WITH_EDITOR
+
+			NeighborOffset += NumNeighbors;
 		}
 	}
-	OutInputData = nullptr;
-	OutNumInputFloats = 0;
 }
 
-void UNearestNeighborModelInstance::GetOutputDataPointer(float*& OutOutputData, int32& OutNumOutputFloats)
+TOptional<TArrayView<float>> UNearestNeighborModelInstance::GetInputView() const
 {
-	const UNearestNeighborModel* NearestNeighborModel = Cast<UNearestNeighborModel>(Model);
+	const UNearestNeighborModel* const NearestNeighborModel = GetCastModel();
 	check(NearestNeighborModel);
-	if (NearestNeighborModel->DoesUseOptimizedNetwork())
+
+	UConstNetworkPtr OptimizedNetwork = NearestNeighborModel->GetOptimizedNetwork();
+	if (OptimizedNetwork.IsValid() && OptimizedNetworkInstance)
 	{
-		const UNearestNeighborOptimizedNetwork* OptimizedNetwork = NearestNeighborModel->GetOptimizedNetwork();
-		if (OptimizedNetwork && OptimizedNetworkInstance)
-		{
-			OutOutputData = OptimizedNetworkInstance->GetOutputs().GetData();
-			OutNumOutputFloats = OptimizedNetwork->GetNumOutputs();
-			return;
-		}
+		float* InputData = OptimizedNetworkInstance->GetInputs().GetData();
+		const int32 NumInputs = OptimizedNetwork->GetNumInputs();
+		return TArrayView<float>(InputData, NumInputs);
 	}
-	OutOutputData = nullptr;
-	OutNumOutputFloats = 0;
+	return TOptional<TArrayView<float>>();
 }
 
-FString UNearestNeighborModelInstance::CheckCompatibility(USkeletalMeshComponent* InSkelMeshComponent, bool LogIssues)
+TOptional<TArrayView<float>> UNearestNeighborModelInstance::GetOutputView() const
 {
-	ErrorText = FString();
+	const UNearestNeighborModel* const NearestNeighborModel = GetCastModel();
+	check(NearestNeighborModel);
 
-	// If we're not compatible, generate a compatibility string.
-	USkeletalMesh* SkelMesh = InSkelMeshComponent ? InSkelMeshComponent->GetSkeletalMeshAsset() : nullptr;
-	UMLDeformerInputInfo* InputInfo = Model->GetInputInfo();
-	if (SkelMesh && !InputInfo->IsCompatible(SkelMesh) && Model->GetDeformerAsset())
+	UConstNetworkPtr OptimizedNetwork = NearestNeighborModel->GetOptimizedNetwork();
+	if (OptimizedNetwork.IsValid() && OptimizedNetworkInstance)
 	{
-		ErrorText += InputInfo->GenerateCompatibilityErrorString(SkelMesh);
-		ErrorText += "\n";
-		if (LogIssues)
-		{
-			UE_LOG(LogNearestNeighborModel, Error, TEXT("ML Deformer '%s' isn't compatible with Skeletal Mesh '%s'.\nReason(s):\n%s"), 
-				*Model->GetDeformerAsset()->GetName(), 
-				*SkelMesh->GetName(), 
-				*ErrorText);
-		}
+		float* OutputData = OptimizedNetworkInstance->GetOutputs().GetData();
+		const int32 NumOutputs = OptimizedNetwork->GetNumOutputs();
+		return TArrayView<float>(OutputData, NumOutputs);
 	}
+	return TOptional<TArrayView<float>>();
+}
 
-	// Verify the number of inputs versus the expected number of inputs.
-	const UNearestNeighborModel* NearestNeighborModel = Cast<UNearestNeighborModel>(Model);
+TArray<float> UNearestNeighborModelInstance::Eval(const TArray<float>& InputData)
+{
+	const UNearestNeighborModel* const NearestNeighborModel = GetCastModel();
 	check(NearestNeighborModel);
-	const int32 NumFloatsPerBone = NearestNeighborModel->GetNumFloatsPerBone();
-	const int32 NumFloatsPerCurve = NearestNeighborModel->GetNumFloatsPerCurve();
-	const int64 NumDeformerAssetInputs = static_cast<int64>(Model->GetInputInfo()->CalcNumNeuralNetInputs(NumFloatsPerBone, NumFloatsPerCurve));
-	int64 NumNeuralNetInputs = -1;
-	if (NearestNeighborModel->DoesUseOptimizedNetwork())
+	TArray<float> Empty;
+	TOptional<TArrayView<float>> OptionalInputView = GetInputView();
+	if (!OptionalInputView.IsSet())
 	{
-		const UNearestNeighborOptimizedNetwork* OptimizedNetwork = NearestNeighborModel->GetOptimizedNetwork();
-		if (OptimizedNetwork)
-		{
-			NumNeuralNetInputs = OptimizedNetwork->GetNumInputs();
-		}
-
+		return Empty;
 	}
+	TArrayView<float> InputView = OptionalInputView.GetValue();
+	const int32 NumNeuralNetInputs = InputView.Num();
 
-	// Only check compatibility after the network is loaded. 
-	if (NumNeuralNetInputs >= 0 && NumNeuralNetInputs != NumDeformerAssetInputs) 
+	if (InputData.Num() != NumNeuralNetInputs)
 	{
-		const FString InputErrorString = "The number of network inputs doesn't match the asset. Please retrain the asset."; 
-		ErrorText += InputErrorString + "\n";
-		if (LogIssues && Model->GetDeformerAsset())
-		{
-			UE_LOG(LogNearestNeighborModel, Error, TEXT("Deformer '%s': %s"), *(Model->GetDeformerAsset()->GetName()), *InputErrorString);
-		}
+		return Empty;
 	}
 
-	return ErrorText;
-}
-
-TArray<float> UNearestNeighborModelInstance::Eval(const TArray<float>& InputData)
-{
-	const UNearestNeighborModel* NearestNeighborModel = Cast<UNearestNeighborModel>(Model);
-	int32 NumNeuralNetInputs = 0;
-	float* InputDataPointer = nullptr;
-	GetInputDataPointer(InputDataPointer, NumNeuralNetInputs);
-	check(InputDataPointer != nullptr);
-	check(InputData.Num() == NumNeuralNetInputs);
-
 	for (int32 i = 0; i < NumNeuralNetInputs; ++i)
 	{
-		InputDataPointer[i] = InputData[i];
+		InputView[i] = InputData[i];
 	}
-	NearestNeighborModel->ClipInputs(InputDataPointer, NumNeuralNetInputs);
+	NearestNeighborModel->ClipInputs(InputView);
 
-	if (NearestNeighborModel->DoesUseOptimizedNetwork())
-	{
-		check(OptimizedNetworkInstance != nullptr);
-		OptimizedNetworkInstance->Run();
-		TArrayView<float> Outputs = OptimizedNetworkInstance->GetOutputs();
-		return TArray<float>(Outputs.GetData(), Outputs.Num());
-	}
-	return TArray<float>();
+	check(OptimizedNetworkInstance);
+	OptimizedNetworkInstance->Run();
+	TArrayView<float> Outputs = OptimizedNetworkInstance->GetOutputs();
+	return TArray<float>(Outputs.GetData(), Outputs.Num());
 }
\ No newline at end of file
