Nearestneighbormodel: refactor

diff --git a/Engine/Plugins/Animation/MLDeformer/NearestNeighborModel/Source/NearestNeighborModel/Private/NearestNeighborOptimizedNetwork.cpp b/Engine/Plugins/Animation/MLDeformer/NearestNeighborModel/Source/NearestNeighborModel/Private/NearestNeighborOptimizedNetwork.cpp
--- a/Engine/Plugins/Animation/MLDeformer/NearestNeighborModel/Source/NearestNeighborModel/Private/NearestNeighborOptimizedNetwork.cpp
+++ b/Engine/Plugins/Animation/MLDeformer/NearestNeighborModel/Source/NearestNeighborModel/Private/NearestNeighborOptimizedNetwork.cpp
@@ -192,9 +192,9 @@ void UNearestNeighborOptimizedNetwork::SetNumOutputs(int32 InNumOutputs)
 	NumOutputs = InNumOutputs;
 }
 
-UNearestNeighborOptimizedNetworkInstance* UNearestNeighborOptimizedNetwork::CreateInstance()
+UNearestNeighborOptimizedNetworkInstance* UNearestNeighborOptimizedNetwork::CreateInstance(UObject* Parent) const
 {
-	UNearestNeighborOptimizedNetworkInstance* Instance = NewObject<UNearestNeighborOptimizedNetworkInstance>(this);
+	UNearestNeighborOptimizedNetworkInstance* Instance = NewObject<UNearestNeighborOptimizedNetworkInstance>(Parent);
 	Instance->Init(this);
 	return Instance;
 }
@@ -238,7 +238,7 @@ const UNearestNeighborOptimizedNetwork* UNearestNeighborOptimizedNetworkInstance
 	return Network.Get();
 }
 
-void UNearestNeighborOptimizedNetworkInstance::Init(UNearestNeighborOptimizedNetwork* InNeuralNetwork)
+void UNearestNeighborOptimizedNetworkInstance::Init(const UNearestNeighborOptimizedNetwork* InNeuralNetwork)
 {
 	Network = InNeuralNetwork;
 	Inputs.SetNumZeroed(Network->GetNumInputs());
diff --git a/Engine/Plugins/Animation/MLDeformer/NearestNeighborModel/Source/NearestNeighborModel/Private/NearestNeighborOptimizedNetwork.h b/Engine/Plugins/Animation/MLDeformer/NearestNeighborModel/Source/NearestNeighborModel/Private/NearestNeighborOptimizedNetwork.h
--- a/Engine/Plugins/Animation/MLDeformer/NearestNeighborModel/Source/NearestNeighborModel/Private/NearestNeighborOptimizedNetwork.h
+++ b/Engine/Plugins/Animation/MLDeformer/NearestNeighborModel/Source/NearestNeighborModel/Private/NearestNeighborOptimizedNetwork.h
@@ -102,7 +102,8 @@ public:
 	 * Create an instance of this neural network.
 	 * @return A pointer to the neural network instance.
 	 */
-	virtual UNearestNeighborOptimizedNetworkInstance* CreateInstance();
+	virtual UNearestNeighborOptimizedNetworkInstance* CreateInstance(UObject* Parent = nullptr) const;
+
 
 	/**
 	 * Get the number of inputs, which is the number of floats the network takes as input.
@@ -218,7 +219,7 @@ protected:
 	 * This is automatically called by the UNearestNeighborOptimizedNetwork::CreateInstance() method.
 	 * @param InNeuralNetwork The network that this is an instance of.
 	 */
-	void Init(UNearestNeighborOptimizedNetwork* InNeuralNetwork);
+	void Init(const UNearestNeighborOptimizedNetwork* InNeuralNetwork);
 
 protected:
 	/** The input values. */
@@ -232,5 +233,5 @@ protected:
 
 	/** The neural network this is an instance of. */
 	UPROPERTY(Transient)
-	TObjectPtr<UNearestNeighborOptimizedNetwork> Network;
+	TObjectPtr<const UNearestNeighborOptimizedNetwork> Network;
 };
diff --git a/Engine/Plugins/Animation/MLDeformer/NearestNeighborModel/Source/NearestNeighborModel/Public/NearestNeighborModel.h b/Engine/Plugins/Animation/MLDeformer/NearestNeighborModel/Source/NearestNeighborModel/Public/NearestNeighborModel.h
--- a/Engine/Plugins/Animation/MLDeformer/NearestNeighborModel/Source/NearestNeighborModel/Public/NearestNeighborModel.h
+++ b/Engine/Plugins/Animation/MLDeformer/NearestNeighborModel/Source/NearestNeighborModel/Public/NearestNeighborModel.h
@@ -2,95 +2,27 @@
 #pragma once
 
 #include "CoreTypes.h"
-#include "MLDeformerModel.h"
-#include "MLDeformerVizSettings.h"
-#include "MLDeformerInputInfo.h"
+
 #include "MLDeformerMorphModel.h"
-#include "NearestNeighborModelVizSettings.h"
-#include "UObject/Object.h"
+#include "NearestNeighborModelHelpers.h"
+
 #include "NearestNeighborModel.generated.h"
 
-struct FSkelMeshImportedMeshInfo;
-class USkeletalMesh;
-class UGeometryCache;
 class UAnimSequence;
-class UMLDeformerAsset;
-class USkeleton;
-class IPropertyHandle;
-class UNearestNeighborModelInstance;
+class UGeometryCache;
 class UNearestNeighborOptimizedNetwork;
-class UNearestNeighborTrainingModel;
+struct FMLDeformerGeomCacheTrainingInputAnim;
 
-namespace UE::NearestNeighborModel
-{
-	class FNearestNeighborEditorModel;
-	class FNearestNeighborGeomCacheSampler;
-	class FNearestNeighborModelDetails;
-};
-
-NEARESTNEIGHBORMODEL_API DECLARE_LOG_CATEGORY_EXTERN(LogNearestNeighborModel, Log, All);
 
 namespace UE::NearestNeighborModel
 {
-	template<class T>
-	class FArrayBufferWithSRV : public FVertexBufferWithSRV
-	{
-	public:
-		void Init(TArray<T>* InArray, const FString& InDebugName);
-
-	private:
-		void InitRHI(FRHICommandListBase& RHICmdList) override;
-
-		TArray<T>* Array = nullptr;
-		FString DebugName;
-	};
-
-	template<class T>
-	class FVertexBufferWithUAV : public FVertexBufferWithSRV
-	{
-	public:
-		void Init(int32 InNumElements, const FString& InDebugName);
-
-	private:
-		void InitRHI(FRHICommandListBase& RHICmdList) override;
-
-		int32 NumElements = 0;
-		FString DebugName;
-	};
-};
-
-USTRUCT(BlueprintType, Blueprintable)
-struct FSkeletonCachePair
-{
-	GENERATED_BODY()
-
-	UPROPERTY(EditAnywhere, BlueprintReadWrite, Category = "Nearest Neighbors")
-	TObjectPtr<UAnimSequence> Skeletons = nullptr;
-
-	UPROPERTY(EditAnywhere, BlueprintReadWrite, Category = "Nearest Neighbors")
-	TObjectPtr<UGeometryCache> Cache = nullptr;
-
-	UPROPERTY(EditAnywhere, BlueprintReadWrite, Category = "Nearest Neighbors")
-	TArray<int32> ExcludedFrames;
-};
-
-USTRUCT(BlueprintType, Blueprintable)
-struct FClothPartEditorData
-{
-	GENERATED_BODY()
-
-	UPROPERTY(EditAnywhere, BlueprintReadWrite, meta = (ClampMin = "1", ClampMax = "128"), Category = "Nearest Neighbors")
-	int32 PCACoeffNum = 128;
-
-	UPROPERTY(EditAnywhere, BlueprintReadWrite, Category = "Nearest Neighbors", meta = (DisplayName = "Vertex Map Path (Optional)"))
-	FString VertexMapPath;
-
-	UPROPERTY(EditAnywhere, BlueprintReadWrite, Category = "Nearest Neighbors")
-	int32 MeshIndex = 0;
+#if WITH_EDITOR
+	class FNearestNeighborModelDetails;
+#endif
 };
 
 USTRUCT(BlueprintType, Blueprintable)
-struct FClothPartData
+struct UE_DEPRECATED(5.4, "FClothPartData is deprecated. Use UNearestNeighborModelSection instead.") FClothPartData
 {
 	GENERATED_BODY()
 
@@ -110,7 +42,6 @@ struct FClothPartData
 	UPROPERTY()
 	uint32 NumNeighbors = 0;
 
-#if WITH_EDITORONLY_DATA
 	/** Vertex indices for this cloth part */
 	UPROPERTY()
 	TArray<uint32> VertexMap;
@@ -134,37 +65,207 @@ struct FClothPartData
 	/** Mapping from NeighborCoeffs to AssetNeighborCoeffs */
 	UPROPERTY()
 	TArray<int32> AssetNeighborIndexMap;
-#endif
 
 	UPROPERTY()
 	TArray<float> NeighborCoeffs;
 };
 
-namespace UE::NearestNeighborModel
+/**
+ * The section of the nearest neighbor model.
+ * Each section contains a set of vertices in the original skeletal mesh.
+ * The nearest neighbor search is performed on each section separately.
+ * For example, if a character asset has both shirt and pants mesh, 
+ * user can define sperate two separate sections for shirt and pants, 
+ * and the nearest neighbor search will be performed separately on those two sections.
+ */
+UCLASS()
+class NEARESTNEIGHBORMODEL_API UNearestNeighborModelSection : public UObject
 {
-	enum EUpdateResult : uint8
-	{
-		SUCCESS = 0,
-		ERROR = 1,
-		WARNING = 2
-	};
-
-	static bool HasError(uint8 Flag)
-	{
-		return (Flag & EUpdateResult::ERROR) != 0;
-	}
+	GENERATED_BODY()
+
+public:
+	DECLARE_MULTICAST_DELEGATE_TwoParams(FPropertyChangedDelegate, FPropertyChangedEvent&, UNearestNeighborModelSection&);
+	using EOpFlag = UE::NearestNeighborModel::EOpFlag;
+
+	int32 GetNumPCACoeffs() const;
+	int32 GetNumVertices() const;
+	int32 GetRuntimeNumNeighbors() const;
+
+	bool IsReadyForTraining() const;
+	bool IsReadyForInference() const;
+
+	const TArray<float>& GetNeighborCoeffs() const;
+
+#if WITH_EDITOR
+	virtual void PostEditChangeProperty(FPropertyChangedEvent& PropertyChangedEvent) override;
+	
+	UFUNCTION(BlueprintPure, Category = "Section")
+	int32 GetAssetNumNeighbors() const;
+	
+	const UAnimSequence* GetNeighborPoses() const;
+	const UGeometryCache* GetNeighborMeshes() const;
+	// Avoid mutable getters. Changing NeighborPoses or NeighborMeshes may cause unexpected errors.  
+	UAnimSequence* GetMutableNeighborPoses() const;
+	UGeometryCache* GetMutableNeighborMeshes() const;
+	
+	void SetVertexMapString(const FString& InString);
+
+	int32 GetMeshIndex() const;
+	void SetMeshIndex(int32 Index);
+
+	UFUNCTION(BlueprintPure, Category = "Section")
+	const TArray<int32>& GetVertexMap() const;
+	
+	UFUNCTION(BlueprintPure, Category = "Section")
+	const TArray<float>& GetPCABasis() const;
+	
+	UFUNCTION(BlueprintPure, Category = "Section")
+	const TArray<float>& GetVertexMean() const;
+	
+	UFUNCTION(BlueprintPure, Category = "Section")
+	const TArray<float>& GetAssetNeighborCoeffs() const;
+
+	const TArray<float>& GetAssetNeighborOffsets() const;
+	const TArray<int32>& GetExcludedFrames() const;
+	const TArray<int32>& GetAssetNeighborIndexMap() const;
+
+	FMLDeformerGeomCacheTrainingInputAnim* GetInputAnim() const;
+
+	EOpFlag UpdateForTraining();
+	EOpFlag UpdateForInference();
+	void InvalidateTraining();
+	void InvalidateInference();
+
+	bool UpdateRuntimeNeighbors();
+	bool GetRuntimeNeighborOffsets(TArray<float>& OutNeighborOffsets) const;
+
+	void SetModel(UNearestNeighborModel* InModel);
+	const UNearestNeighborModel* GetModel() const;
+
+	static FName GetNumPCACoeffsPropertyName() { return GET_MEMBER_NAME_CHECKED(UNearestNeighborModelSection, NumPCACoeffs); }
+	static FName GetVertexMapStringPropertyName() { return GET_MEMBER_NAME_CHECKED(UNearestNeighborModelSection, VertexMapString); }
+	static FName GetNeighborPosesPropertyName() { return GET_MEMBER_NAME_CHECKED(UNearestNeighborModelSection, NeighborPoses); }
+	static FName GetNeighborMeshesPropertyName() { return GET_MEMBER_NAME_CHECKED(UNearestNeighborModelSection, NeighborMeshes); }
+	static FName GetExcludedFramesPropertyName() { return GET_MEMBER_NAME_CHECKED(UNearestNeighborModelSection, ExcludedFrames); }
+#endif
+#if WITH_EDITORONLY_DATA
+PRAGMA_DISABLE_DEPRECATION_WARNINGS
+	UE_DEPRECATED(5.4, "InitFromClothPartData is deprecated. This function is only used for backward compatibility.")
+	void InitFromClothPartData(FClothPartData& InPart);
+PRAGMA_ENABLE_DEPRECATION_WARNINGS
+#endif
+
+protected:
+#if WITH_EDITORONLY_DATA
+
+	/** Poses of the nearest neighbor ROM. */
+	UPROPERTY(EditAnywhere, Category = "Section")
+	TObjectPtr<UAnimSequence> NeighborPoses;
+
+	/** Geometry cache of the nearest neighbor ROM. */
+	UPROPERTY(EditAnywhere, Category = "Section")
+	TObjectPtr<UGeometryCache> NeighborMeshes;
+
+	/** A string containing vertex indices for this section, e.g. "2, 3, 5-8, 9, 11-20" */
+	UPROPERTY(EditAnywhere, Category = "Section", meta = (DisplayName = "Vertex Indices"))
+	FString VertexMapString;
+
+	/** Frames to be excluded from the nearest neighbor ROM */
+	UPROPERTY(EditAnywhere, Category = "Section")
+	TArray<int32> ExcludedFrames;
+
+	/** The mesh index in SkeletalMeshRenderData */
+	UPROPERTY()
+	int32 MeshIndex = 0;
+
+	/** The parent nearest neighbor model */
+	UPROPERTY()
+	TObjectPtr<UNearestNeighborModel> Model = nullptr;
+
+	/** Vertex indices for this section */
+	UPROPERTY()
+	TArray<int32> VertexMap;
+	
+	/** Flattened array of PCA basis. The shape of PCA basis is (PCACoeffNum, NumVertices * 3)  */
+	UPROPERTY()
+	TArray<float> PCABasis;
+	
+	/** The vertex mean used in PCA computation. This array has a size of NumVertices * 3 */
+	UPROPERTY()
+	TArray<float> VertexMean;
+
+	/** Flattened array of neighbor PCA coefficients before excluding frames. The shape of this array is (NumNeighbors, NumPCACoeffs) */
+	UPROPERTY()
+	TArray<float> AssetNeighborCoeffs;
+	
+	/** Flattened array of neighbor offsets before excluding frames. The shape of this array is (NumNeighbors, NumVertices * 3) */
+	UPROPERTY()
+	TArray<float> AssetNeighborOffsets;
+
+	/** The index into the original nearest neighbor ROM asset after frames are excluded. */
+	UPROPERTY()
+	TArray<int32> AssetNeighborIndexMap;
+#endif
+
+	/** Number of PCA coefficients for this section. */
+	UPROPERTY(EditAnywhere, BlueprintReadWrite, meta = (ClampMin = "1", DefaultValue = "64"), Category = "Section")
+	int32 NumPCACoeffs = 64;
+
+	/** Number of vertices in this section */
+	UPROPERTY(VisibleAnywhere, BlueprintReadOnly, Category = "Section Private")
+	int32 NumVertices = 0;
+	
+	/** Number of neighbors in this section after excluding frames */
+	UPROPERTY(VisibleAnywhere, BlueprintReadOnly, Category = "Section Private")
+	int32 RuntimeNumNeighbors = 0;
+	
+	/** Flattened array of neighbor offsets after excluding frames. */
+	UPROPERTY()
+	TArray<float> RuntimeNeighborCoeffs;
+
+private:
+	/** Whether the section is ready for training */
+	UPROPERTY()
+	bool bIsReadyForTraining = false;
+	
+	/** Whether the section is ready for inference */
+	UPROPERTY()
+	bool bIsReadyForInference = false;
+
+#if WITH_EDITOR
+	bool IsPCAValid() const;
+	bool IsPCAEmpty() const;
+	bool IsNearestNeighborValid() const;
+	bool IsNearestNeighborEmpty() const;
+
+	UFUNCTION(BlueprintCallable, Category = "Section")
+	void SetPCAData(const TArray<float>& InVertexMean, const TArray<float>& InPCABasis);
+	
+	UFUNCTION(BlueprintCallable, Category = "Section")
+	void SetNeighborData(const TArray<float>& InNeighborCoeffs, const TArray<float>& InNeighborOffsets);
+
+	void Reset();
+	void ResetPCAData();
+	void ResetNearestNeighborData();
+	EOpFlag UpdateVertexMap();
+	
+	static FString BackwardCompatibleVertexString;
+
+	/** A temporary InputAnim used by GetInputAnim() function. */
+	mutable TUniquePtr<FMLDeformerGeomCacheTrainingInputAnim> InputAnim;
+#endif
 };
 
 /**
  * The nearest neighbor model.
- * This model contains the PCA basis of the cloth vertex deltas and a small set of cloth for nearest neighbor search. 
+ * This model contains the PCA basis of the vertex deltas and a small set of meshes for nearest neighbor search. 
  * Given a new pose, the pre-trained neural network first predicts the PCA coefficients of the vertex deltas. 
- * Then this model uses the predicted PCA coeffcients to find a nearest neighbor in the small cloth set.
+ * Then this model uses the predicted PCA coeffcients to find a nearest neighbor in the small dataset.
  * The total vertex delta is computed by
  * 		vertex_delta = mean_delta + pca_basis * pca_coeff + nearest_neighbor_delta
  * To prevent popping, a time filtering is applied on predicted vertex deltas. The vertex delta at time t is computed by  
  * 		vertex_delta(t) = decay_factor * vertex_delta(t-1) + (1 - decay_factor) * vertex_delta 
- * The cloth can be separated into several parts (e.g. shirt, pants...). The nearest neighbor search is carried out separately for each part. 
+ * The mesh can be separated into several sections (e.g. shirt, pants...). The nearest neighbor search is carried out separately for each section. 
  * The pca basis and the nearest neighbor data are compressed into morph targets.
  */
 UCLASS()
@@ -174,10 +275,16 @@ class NEARESTNEIGHBORMODEL_API UNearestNeighborModel
 	GENERATED_BODY()
 
 public:
+	using FSection = UNearestNeighborModelSection;
+	using UNetwork = UNearestNeighborOptimizedNetwork;
+
 	UNearestNeighborModel(const FObjectInitializer& ObjectInitializer);
 
 	// UObject overrides.
 	virtual void GetAssetRegistryTags(TArray<FAssetRegistryTag>& OutTags) const override;
+#if WITH_EDITOR
+	virtual void PostEditChangeProperty(FPropertyChangedEvent& PropertyChangedEvent) override;
+#endif
 	// ~END UObject overrides.
 
 	// UMLDeformerModel overrides.
@@ -186,293 +293,266 @@ public:
 	virtual bool DoesSupportQualityLevels() const override { return false; }
 	virtual UMLDeformerModelInstance* CreateModelInstance(UMLDeformerComponent* Component) override;
 	virtual UMLDeformerInputInfo* CreateInputInfo() override;
-	virtual FString GetDisplayName() const override { return "Nearest Neighbor Model"; }
+	virtual FString GetDisplayName() const override;
+	virtual int32 GetNumFloatsPerBone() const override { return NearestNeighborNumFloatsPerBone; }
+	virtual int32 GetNumFloatsPerCurve() const override { return NearestNeighborNumFloatsPerCurve; }
 	virtual bool IsTrained() const override;
+	virtual FString GetDefaultDeformerGraphAssetPath() const;
 	// ~END UMLDeformerModel overrides.
 
-	friend class UNearestNeighborModelInstance;
-	friend class UNearestNeighborTrainingModel;
-	friend class UE::NearestNeighborModel::FNearestNeighborEditorModel;
-	friend class UE::NearestNeighborModel::FNearestNeighborGeomCacheSampler;
-	friend class UE::NearestNeighborModel::FNearestNeighborModelDetails;
-
-private:
 	UFUNCTION(BlueprintPure, Category = "Nearest Neighbor Model")
-	int32 GetNumParts() const { return ClothPartData.Num(); }
+	int32 GetNumSections() const;
 
-	UFUNCTION(BlueprintPure, Category = "Nearest Neighbor Model")
-	int32 GetPartNumVerts(int32 PartId) const { return ClothPartData[PartId].NumVertices; }
-	
-	UFUNCTION(BlueprintPure, Category = "Nearest Neighbor Model")
-	int32 GetPCACoeffStart(int32 PartId) const { return ClothPartData[PartId].PCACoeffStart; }
+	// GetSectionPtr is reserved for python. Use GetSection for C++.
+	UFUNCTION(BlueprintPure, Category = "Python")
+	const UNearestNeighborModelSection* GetSectionPtr(int32 Index) const;
+	const UNearestNeighborModelSection& GetSection(int32 Index) const;
 
 	UFUNCTION(BlueprintPure, Category = "Nearest Neighbor Model")
-	int32 GetPCACoeffNum(int32 PartId) const { return ClothPartData[PartId].PCACoeffNum; }
-
-	int32 GetNumNeighbors(int32 PartId) const { return ClothPartData[PartId].NumNeighbors; }
+	const TArray<int32>& GetPCACoeffStarts() const;
 
+	UFUNCTION(BlueprintPure, Category = "Nearest Neighbor Model")
 	int32 GetTotalNumPCACoeffs() const;
+
 	int32 GetTotalNumNeighbors() const;
+	float GetDecayFactor() const;
+	float GetNearestNeighborOffsetWeight() const;
 
-	float GetDecayFactor() const { return DecayFactor; };
-	float GetNearestNeighborOffsetWeight() const { return NearestNeighborOffsetWeight; }
+	TWeakObjectPtr<const UNetwork> GetOptimizedNetwork() const;
+	int32 GetNumNetworkOutputs() const;
+	void ClipInputs(TArrayView<float> Inputs) const;
 
 	UFUNCTION(BlueprintPure, Category = "Nearest Neighbor Model")
-	TArray<float> ClipInputs(const TArray<float>& Input) const;
-
-	void ClipInputs(float* InputPtr, int NumInputs) const;
-
-	void InitInputInfo();
-
-	void ResetMorphBuffers();
-
-	const TArray<float>& NeighborCoeffs(int32 PartId) const { return ClothPartData[PartId].NeighborCoeffs; }
-
-	virtual int32 GetNumFloatsPerBone() const override { return NearestNeighborNumFloatsPerBone; }
-	virtual int32 GetNumFloatsPerCurve() const override { return NearestNeighborNumFloatsPerCurve; }
-
-	bool DoesUseOptimizedNetwork() const;
-	bool ShouldUseOptimizedNetwork() const;
-	UNearestNeighborOptimizedNetwork* GetOptimizedNetwork() const { return OptimizedNetwork.Get(); }
-	int32 GetOptimizedNetworkNumOutputs() const;
-	void SetOptimizedNetwork(UNearestNeighborOptimizedNetwork* InOptimizedNetwork);
+	bool IsReadyForTraining() const;
 
-	UE_DEPRECATED(5.4, "Onnx file no longer used for training network. Use LoadOptimizedNetworkFromFile instead.")
-	bool LoadOptimizedNetwork(const FString& OnnxPath);
+	UFUNCTION(BlueprintPure, Category = "Nearest Neighbor Model")
+	bool IsReadyForInference() const;
 
-	bool LoadOptimizedNetworkFromFile(const FString& Filename);
+	bool DoesUseRBF() const;
+	float GetRBFSigma() const;
 
 #if WITH_EDITOR
 	UFUNCTION(BlueprintPure, Category = "Nearest Neighbor Model")
-	int32 GetAssetNumNeighbors(int32 PartId) const;
-
-	const TArray<uint32>& PartVertexMap(int32 PartId) const { return ClothPartData[PartId].VertexMap; }
-
+	int32 GetInputDim() const { return InputDim; }
+	
 	UFUNCTION(BlueprintPure, Category = "Nearest Neighbor Model")
-	const TArray<float>& PCABasis(int32 PartId) const { return ClothPartData[PartId].PCABasis; }
-
-	TArray<float>& PCABasis(int32 PartId) { return ClothPartData[PartId].PCABasis; }
-
-	UFUNCTION(BlueprintCallable, Category = "Nearest Neighbor Model")
-	void SetPCABasis(int32 PartId, const TArray<float>& PCABasis) { ClothPartData[PartId].PCABasis = PCABasis; }
-
+	TArray<int32> GetHiddenLayerDims() const { return HiddenLayerDims; }
+	
+	UFUNCTION(BlueprintPure, Category = "Nearest Neighbor Model")
+	int32 GetOutputDim() const { return OutputDim; }
+	
 	UFUNCTION(BlueprintPure, Category = "Nearest Neighbor Model")
-	const TArray<float>& VertexMean(int32 PartId) const { return ClothPartData[PartId].VertexMean; }
+	int32 GetNumEpochs() const { return NumEpochs; }
+	
+	UFUNCTION(BlueprintPure, Category = "Nearest Neighbor Model")
+	int32 GetBatchSize() const { return BatchSize; }
+	
+	UFUNCTION(BlueprintPure, Category = "Nearest Neighbor Model")
+	float GetLearningRate() const { return LearningRate; }
 
-	TArray<float>& VertexMean(int32 PartId) { return ClothPartData[PartId].VertexMean; }
+	UFUNCTION(BlueprintPure, Category = "Nearest Neighbor Model")
+	int32 GetEarlyStopEpochs() const { return EarlyStopEpochs; }
 
-	UFUNCTION(BlueprintCallable, Category = "Nearest Neighbor Model")
-	void SetVertexMean(int32 PartId, const TArray<float>& VertexMean) { ClothPartData[PartId].VertexMean = VertexMean; }
+	UFUNCTION(BlueprintPure, Category = "Nearest Neighbor Model")
+	FString GetModelDir() const;
+	
+	bool DoesUseFileCache() const { return bUseFileCache; }
+	bool DoesUseDualQuaternionDeltas() const { return bUseDualQuaternionDeltas; }
 
-	void SetNumNeighbors(int32 PartId, int32 InNumNeighbors) { ClothPartData[PartId].NumNeighbors = InNumNeighbors; }
+	FDateTime GetNetworkLastWriteTime() const;
+	FString GetNetworkLastWriteArchitectureString() const;
+	FDateTime GetMorphTargetsLastWriteTime() const;
 
-	UFUNCTION(BlueprintCallable, Category = "Nearest Neighbor Model")
-	void SetAssetNeighborCoeffs(int32 PartId, const TArray<float>& AssetNeighborCoeffs) { ClothPartData[PartId].AssetNeighborCoeffs = AssetNeighborCoeffs; }
+	UFUNCTION(BlueprintPure, Category = "Nearest Neighbor Model")
+	TArray<FString> GetCachedDeltasPaths() const;
 
-	void SetNeighborCoeffs(int32 PartId, const TArray<float>& NeighborCoeffs) { ClothPartData[PartId].NeighborCoeffs = NeighborCoeffs; }
+	UFUNCTION(BlueprintPure, Category = "Nearest Neighbor Model")
+	TArray<FString> GetCachedPCAPaths() const;
 
 	UFUNCTION(BlueprintPure, Category = "Nearest Neighbor Model")
-	const TArray<float>& AssetNeighborCoeffs(int32 PartId) const { return ClothPartData[PartId].AssetNeighborCoeffs; }
+	TArray<FString> GetCachedNetworkPaths() const;
 	
-	const TArray<float>& AssetNeighborOffsets(int32 PartId) const { return ClothPartData[PartId].AssetNeighborOffsets; }
+	TOptional<FDateTime> GetCachedDeltasTimestamp() const;
+	TOptional<FDateTime> GetCachedPCATimestamp() const;
+	TOptional<FDateTime> GetCachedNetworkTimestamp() const;
+
+	using EOpFlag = UE::NearestNeighborModel::EOpFlag;
+	EOpFlag UpdateForTraining();
+	void InvalidateTraining();
+	void InvalidateTrainingModelOnly();
+	EOpFlag UpdateForInference();
+	void InvalidateInference();
+	void InvalidateInferenceModelOnly();
+	bool LoadOptimizedNetworkFromFile(const FString& Filename);
+	void ClearOptimizedNetwork();
 
-	UFUNCTION(BlueprintCallable, Category = "Nearest Neighbor Model")
-	void SetAssetNeighborOffsets(int32 PartId, const TArray<float>& AssetNeighborOffsets) { ClothPartData[PartId].AssetNeighborOffsets = AssetNeighborOffsets; }
-#endif
+	FMLDeformerGeomCacheTrainingInputAnim* GetNearestNeighborAnim(int32 SectionIndex);
+	const FMLDeformerGeomCacheTrainingInputAnim* GetNearestNeighborAnim(int32 SectionIndex) const;
 
-#if WITH_EDITORONLY_DATA
-	TObjectPtr<UAnimSequence> GetNearestNeighborSkeletons(int32 PartId);
-	const TObjectPtr<UAnimSequence> GetNearestNeighborSkeletons(int32 PartId) const;
-	TObjectPtr<UGeometryCache> GetNearestNeighborCache(int32 PartId);
-	const TObjectPtr<UGeometryCache> GetNearestNeighborCache(int32 PartId) const;
-	int32 GetNumNeighborsFromGeometryCache(int32 PartId) const;
-	int32 GetNumNeighborsFromAnimSequence(int32 PartId) const;
+	void UpdateFileCache();
+	const FString& GetFileCacheDirectory() const;
+	void SetFileCacheDirectory(const FString& InFileCacheDirectory);
+
+	/** Make sure to call this function after changing the morph targets.
+	* Do not call this function without calling InitEngineMorphTargets first. */
+	void UpdateMorphTargetsLastWriteTime();
 
 	void UpdateNetworkInputDim();
 	void UpdateNetworkOutputDim();
-	bool CheckPCAData(int32 PartId) const;
-	UE::NearestNeighborModel::EUpdateResult UpdateClothPartData();
-	UE::NearestNeighborModel::EUpdateResult UpdateVertexMap(int32 PartId, const FString& VertexMapPath, const FSkelMeshImportedMeshInfo& Info);
-	void UpdatePCACoeffNums();
-	void UpdateNetworkSize();
-	void UpdateMorphTargetSize();
-	void UpdateInputMultipliers();
 
-	void InvalidateClothPartData() { bClothPartDataValid = false; bNearestNeighborDataValid = false; bMorphTargetDataValid = false; }
-	void ValidateClothPartData() { bClothPartDataValid = true; }
-	void InvalidateNearestNeighborData() { bNearestNeighborDataValid = false; bMorphTargetDataValid = false; }
-	void ValidateNearestNeighborData() {bNearestNeighborDataValid = true; }
-	void InvalidateMorphTargetData() { bMorphTargetDataValid = false; }
-	void ValidateMorphTargetData() { bMorphTargetDataValid = true;}
-	bool IsClothPartDataValid() const { return bClothPartDataValid; }
-	bool IsNearestNeighborDataValid() const { return bNearestNeighborDataValid; }
-	bool IsMorphTargetDataValid() const { return bMorphTargetDataValid; }
-	void SetUseOptimizedNetwork(bool bInUseOptimizedNetwork);
-#endif
+	bool IsBeforeCustomVersionWasAdded() const;
 
-#if WITH_EDITOR
-	UFUNCTION(BlueprintPure, Category = "Nearest Neighbor Model")
-	const int32	GetInputDim() const { return InputDim; }
 	static FName GetInputDimPropertyName() { return GET_MEMBER_NAME_CHECKED(UNearestNeighborModel, InputDim); }
-
-	UFUNCTION(BlueprintPure, Category = "Nearest Neighbor Model")
-	const TArray<int32>& GetHiddenLayerDims() const { return HiddenLayerDims; }
 	static FName GetHiddenLayerDimsPropertyName() { return GET_MEMBER_NAME_CHECKED(UNearestNeighborModel, HiddenLayerDims); }
-
-	UFUNCTION(BlueprintPure, Category = "Nearest Neighbor Model")
-	const int32	GetOutputDim() const { return OutputDim; }
 	static FName GetOutputDimPropertyName() { return GET_MEMBER_NAME_CHECKED(UNearestNeighborModel, OutputDim); }
-
-	UFUNCTION(BlueprintPure, Category = "Nearest Neighbor Model")
-	const int32	GetNumEpochs() const { return NumEpochs; }
 	static FName GetNumEpochsPropertyName() { return GET_MEMBER_NAME_CHECKED(UNearestNeighborModel, NumEpochs); }
-
-	UFUNCTION(BlueprintPure, Category = "Nearest Neighbor Model")
-	const int32	GetBatchSize() const { return BatchSize; }
 	static FName GetBatchSizePropertyName() { return GET_MEMBER_NAME_CHECKED(UNearestNeighborModel, BatchSize); }
-
-	UFUNCTION(BlueprintPure, Category = "Nearest Neighbor Model")
-	const float	GetLearningRate() const { return LearningRate; }
 	static FName GetLearningRatePropertyName() { return GET_MEMBER_NAME_CHECKED(UNearestNeighborModel, LearningRate); }
-
-	UFUNCTION(BlueprintPure, Category = "Nearest Neighbor Model")
-	const TArray<FClothPartEditorData>&	GetClothPartEditorData() const { return ClothPartEditorData; }
-	static FName GetClothPartEditorDataPropertyName() { return GET_MEMBER_NAME_CHECKED(UNearestNeighborModel, ClothPartEditorData); }
-
-	UFUNCTION(BlueprintPure, Category = "Nearest Neighbor Model")
-	const TArray<FSkeletonCachePair>& GetNearestNeighborData() const { return NearestNeighborData; }
-	static FName GetNearestNeighborDataPropertyName() { return GET_MEMBER_NAME_CHECKED(UNearestNeighborModel, NearestNeighborData); }
-
+	static FName GetEarlyStopEpochsPropertyName() { return GET_MEMBER_NAME_CHECKED(UNearestNeighborModel, EarlyStopEpochs); }
+	static FName GetSectionsPropertyName() { return GET_MEMBER_NAME_CHECKED(UNearestNeighborModel, Sections); }
 	static FName GetNearestNeighborOffsetWeightPropertyName() { return GET_MEMBER_NAME_CHECKED(UNearestNeighborModel, NearestNeighborOffsetWeight); }
-
-	static FName GetSavedNetworkSizePropertyName() { return GET_MEMBER_NAME_CHECKED(UNearestNeighborModel, SavedNetworkSize); }
-	static FName GetMorphDataSizePropertyName() { return GET_MEMBER_NAME_CHECKED(UNearestNeighborModel, MorphDataSize); }
-
-	bool GetUseFileCache() const { return bUseFileCache; }
 	static FName GetUseFileCachePropertyName() { return GET_MEMBER_NAME_CHECKED(UNearestNeighborModel, bUseFileCache); }
 	static FName GetFileCacheDirectoryPropertyName() { return GET_MEMBER_NAME_CHECKED(UNearestNeighborModel, FileCacheDirectory); }
-	static FName GetRecomputeDeltasPropertyName() { return GET_MEMBER_NAME_CHECKED(UNearestNeighborModel, bRecomputeDeltas); }
-	static FName GetRecomputePCAPropertyName() { return GET_MEMBER_NAME_CHECKED(UNearestNeighborModel, bRecomputePCA); }
-
-	UFUNCTION(BlueprintPure, Category = "Nearest Neighbor Model")
-	FString GetModelDir() const;
-
-	int32 GetPartMeshIndex(int32 PartId) const { return ClothPartEditorData[PartId].MeshIndex; }
-	void SetPartMeshIndex(int32 PartId, int32 MeshIndex) { ClothPartEditorData[PartId].MeshIndex = MeshIndex; }
-	int32 GetMaxPartMeshIndex() const;
-
-	bool DoesUseDualQuaternionDeltas() const { return bUseDualQuaternionDeltas; }
+	static FName GetUseDualQuaternionDeltasPropertyName() { return GET_MEMBER_NAME_CHECKED(UNearestNeighborModel, bUseDualQuaternionDeltas); }
+	static FName GetDecayFactorPropertyName() { return GET_MEMBER_NAME_CHECKED(UNearestNeighborModel, DecayFactor); }
+	static FName GetUseRBFPropertyName() { return GET_MEMBER_NAME_CHECKED(UNearestNeighborModel, bUseRBF); }
+	static FName GetRBFSigmaPropertyName() { return GET_MEMBER_NAME_CHECKED(UNearestNeighborModel, RBFSigma); }
 #endif
-
+	
 protected:
-
+	/** Network input dimensions. */
 	UPROPERTY(VisibleAnywhere, BlueprintReadOnly, Category = "Training Settings")
 	int32 InputDim = 0;
 
+	/** Network output dimensions. */
 	UPROPERTY(VisibleAnywhere, BlueprintReadOnly, Category = "Training Settings")
 	int32 OutputDim = 0;
 
 #if WITH_EDITORONLY_DATA
+	/** Dimension of hidden layers in the network. This cannot be empty. */
 	UPROPERTY(EditAnywhere, BlueprintReadWrite, Category = "Training Settings")
 	TArray<int32> HiddenLayerDims;
 
-	UPROPERTY(EditAnywhere, BlueprintReadWrite, Category = "Training Settings", meta = (ClampMin = "1"))
+	/** Max number of cycles iterated through the training set. */
+	UPROPERTY(EditAnywhere, BlueprintReadWrite, Category = "Training Settings", AdvancedDisplay, meta = (ClampMin = "1"))
 	int32 NumEpochs = 10000;
 
-	UPROPERTY(EditAnywhere, BlueprintReadWrite, Category = "Training Settings", meta = (ClampMin = "1"))
+	/** Number of data samples processed together as a group in a single pass. */
+	UPROPERTY(EditAnywhere, BlueprintReadWrite, Category = "Training Settings", AdvancedDisplay, meta = (ClampMin = "1"))
 	int32 BatchSize = 256;
 
-	UPROPERTY(EditAnywhere, BlueprintReadWrite, Category = "Training Settings", meta = (ClampMin = "0.000001", ClampMax = "1.0"))
+	/** The size of the step when optimizing the network. */
+	UPROPERTY(EditAnywhere, BlueprintReadWrite, Category = "Training Settings", AdvancedDisplay, meta = (ClampMin = "0.000001", ClampMax = "1.0"))
 	float LearningRate = 0.001f;
 
-	UPROPERTY()
-	bool bClothPartDataValid = false;
-
-	UPROPERTY(EditAnywhere, BlueprintReadWrite, Category = "Cloth Parts")
-	TArray<FClothPartEditorData> ClothPartEditorData;
+	/** The number of epochs to stop training if there is no improvement in accuracy. */
+	UPROPERTY(EditAnywhere, BlueprintReadWrite, Category = "Training Settings", AdvancedDisplay, meta = (ClampMin = "1"))
+	int32 EarlyStopEpochs = 100;
 
-	UPROPERTY(EditAnywhere, BlueprintReadWrite, Category = "Cloth Parts", meta = (ClampMin = "0"))
-	int32 BasisSmoothIter = 0;
-
-	UPROPERTY(EditAnywhere, BlueprintReadWrite, Category = "Cloth Parts")
-	bool bUseDualQuaternionDeltas = false;
-
-	UPROPERTY()
-	bool bNearestNeighborDataValid = false;
-
-	UPROPERTY()
-	bool bMorphTargetDataValid = false;
-
-	UPROPERTY(EditAnywhere, BlueprintReadWrite, Category = "Nearest Neighbors")
-	TArray<FSkeletonCachePair> NearestNeighborData;
-
-	UPROPERTY(VisibleAnywhere, BlueprintReadOnly, Category = "Training Settings", meta = (ClampMin = "0.0"))
-	float SavedNetworkSize = 0.0f;
-
-	UPROPERTY(VisibleAnywhere, BlueprintReadOnly, Category = "Morph Targets", meta = (ClampMin = "0.0"))
-	float MorphDataSize = 0.0f;
-
-	UPROPERTY(EditAnywhere, BlueprintReadWrite, Category = "File Cache")
+	/** Whether to cache intermediate results on disk. CAUTION: failing to manually clear caches could cause unexpected results. */
+	UPROPERTY(EditAnywhere, BlueprintReadWrite, Category = "File Cache", meta = (DisplayName = "Use File Cache (Advanced)"))
 	bool bUseFileCache = false;
 
-	UPROPERTY(EditAnywhere, BlueprintReadWrite, Category = "File Cache")
+	/** Directory to save the intermediate results. */
+	UPROPERTY(VisibleAnywhere, BlueprintReadOnly, Category = "File Cache", meta = (EditCondition = "bUseFileCache"))
 	FString FileCacheDirectory = FPaths::ProjectIntermediateDir() + "NearestNeighborModel";
 
-	UPROPERTY(EditAnywhere, BlueprintReadWrite, Category = "File Cache")
-	bool bRecomputeDeltas = true;
-
-	UPROPERTY(EditAnywhere, BlueprintReadWrite, Category = "File Cache")
-	bool bRecomputePCA = true;
+	TOptional<FDateTime> CachedDeltasTimestamp;
+	TOptional<FDateTime> CachedPCATimestamp;
+	TOptional<FDateTime> CachedNetworkTimestamp;
 
+PRAGMA_DISABLE_DEPRECATION_WARNINGS
+	UE_DEPRECATED(5.4, "ClothPartData is deprecated. Use Sections instead.")
+	UPROPERTY()
+	TArray<FClothPartData> ClothPartData_DEPRECATED;
+PRAGMA_ENABLE_DEPRECATION_WARNINGS
 #endif
 
 	static constexpr int32 NearestNeighborNumFloatsPerBone = 3;
 	static constexpr int32 NearestNeighborNumFloatsPerCurve = 0;
-	
-#if WITH_EDITORONLY_DATA
-	UPROPERTY(EditAnywhere, BlueprintReadWrite, Category = "KMeans Pose Generator")
-	TArray<TObjectPtr<UAnimSequence>> SourceAnims;
-
-	UPROPERTY(EditAnywhere, BlueprintReadWrite, Category = "KMeans Pose Generator", meta = (ClampMin = "1"))
-	int32 NumClusters;
-
-	UPROPERTY(EditAnywhere, BlueprintReadWrite, Category = "KMeans Pose Generator")
-	int32 KMeansPartId = 0;
-#endif
 
+	/** Whether to use input multipliers. This can be used to debug bad network input. */
 	UPROPERTY(EditAnywhere, BlueprintReadWrite, Category = "Debugging")
 	bool bUseInputMultipliers = false;
 
+	/** Values to be multiplied to the input. This can be used to debug bad network input. */
 	UPROPERTY(EditAnywhere, BlueprintReadWrite, Category = "Debugging")
 	TArray<FVector3f> InputMultipliers;
 
-	UPROPERTY()
-	TArray<FClothPartData> ClothPartData;
-
+	/** The min input values observed throughout the entire training set. This is used to clamp the input value at inference time.  The values are set in python. */
 	UPROPERTY(BlueprintReadWrite, Category = "Network IO")
 	TArray<float> InputsMin;
 
+	/** The max input values observed throughout the entire training set. This is used to clamp the input value at inference time.  The values are set in python. */
 	UPROPERTY(BlueprintReadWrite, Category = "Network IO")
 	TArray<float> InputsMax;
 
-	UPROPERTY(EditAnywhere, BlueprintReadWrite, Category = "Nearest Neighbors", META = (ClampMin = "0", ClampMax = "1"))
+	/** Whether to use dual quaternion deltas. If false, LBS deltas will be used. */
+	UPROPERTY(EditAnywhere, BlueprintReadWrite, Category = "Nearest Neighbor Settings")
+	bool bUseDualQuaternionDeltas = true;
+
+	/** The ratio of previous frame deltas added into the current frame deltas. A bigger value will make wrinkles "stick" longer. */
+	UPROPERTY(EditAnywhere, BlueprintReadWrite, Category = "Nearest Neighbor Settings", META = (ClampMin = "0", ClampMax = "1"))
 	float DecayFactor = 0.85f;
 
-	UPROPERTY(EditAnywhere, BlueprintReadWrite, Category = "Nearest Neighbors", META = (ClampMin = "0", ClampMax = "1"))
+	/** The weight multiplied to nearest neighbor deltas. A value of 0 means completely removing nearest neighbor deltas. */
+	UPROPERTY(EditAnywhere, BlueprintReadWrite, Category = "Nearest Neighbor Settings", META = (ClampMin = "0", ClampMax = "1"))
 	float NearestNeighborOffsetWeight = 1.0f;
 
-	/** Use radial basis function (RBF). This will blend multiple nearest neighbors instead of finding a single neighbor. */
-	UPROPERTY(EditAnywhere, BlueprintReadWrite, Category = "Nearest Neighbors")
+	/** Whether to use radial basis function to blend multiple nearest neighbors to produce smoother result. */
+	UPROPERTY(EditAnywhere, BlueprintReadWrite, Category = "Nearest Neighbor Settings")
 	bool bUseRBF = false;
 
-	/** The smaller this value is, the more similar it is to a single nearest neighbor. The bigger this value is, the more smeared out the geometric details are. */
-	UPROPERTY(EditAnywhere, BlueprintReadWrite, Category = "Nearest Neighbors", META = (ClampMin = "0.001"))
+	/** Range to blend nearest neighbors. A bigger range will blend more neighbors, produce smoother result and may be slower in game. */
+	UPROPERTY(EditAnywhere, BlueprintReadWrite, Category = "Nearest Neighbor Settings", META = (ClampMin = "0"))
 	float RBFSigma = 1.0f;
 
 private:
+	TWeakObjectPtr<UNetwork> GetOptimizedNetwork();
+
+#if WITH_EDITOR
+	UNearestNeighborModelSection* OnSectionAdded(int32 NewIndex);
+	FSection& GetSection(int32 Index);
+
+	void SetOptimizedNetwork(UNearestNeighborOptimizedNetwork* InOptimizedNetwork);
+
+	void UpdatePCACoeffStarts();
+
+	EOpFlag CheckHiddenLayerDims();
+	void UpdateInputMultipliers();
+
+	void UpdateCachedDeltasTimestamp();
+	void UpdateCachedPCATimestamp();
+	void UpdateCachedNetworkTimestamp();
+	// FNearestNeighborModelDetails needs to call private function GetSection(int32).
+	friend class UE::NearestNeighborModel::FNearestNeighborModelDetails;
+#endif
+
+	UPROPERTY(EditAnywhere, Category = "Mesh Sections")
+	TArray<TObjectPtr<UNearestNeighborModelSection>> Sections;
+	
+	UPROPERTY()
+	TArray<int32> PCACoeffStarts;
+	
+	UPROPERTY()
+	bool bIsReadyForTraining = false;
+
+	UPROPERTY()
+	bool bIsReadyForInference = false;
+
 	UPROPERTY()
-	TObjectPtr<UNearestNeighborOptimizedNetwork> OptimizedNetwork = nullptr;
+	TObjectPtr<UNearestNeighborOptimizedNetwork> OptimizedNetwork;
+
+	int32 Version = INDEX_NONE;
+
+#if WITH_EDITORONLY_DATA
+	UPROPERTY()
+	int64 NetworkLastWriteTime = 0;
 
 	UPROPERTY()
-	bool bUseOptimizedNetwork = true;
+	int64 MorphTargetsLastWriteTime = 0;
+
+	UPROPERTY()
+	FString NetworkLastWriteArchitectureString;
+#endif
 };
diff --git a/Engine/Plugins/Animation/MLDeformer/NearestNeighborModel/Source/NearestNeighborModel/Public/NearestNeighborModelInputInfo.h b/Engine/Plugins/Animation/MLDeformer/NearestNeighborModel/Source/NearestNeighborModel/Public/NearestNeighborModelInputInfo.h
--- a/Engine/Plugins/Animation/MLDeformer/NearestNeighborModel/Source/NearestNeighborModel/Public/NearestNeighborModelInputInfo.h
+++ b/Engine/Plugins/Animation/MLDeformer/NearestNeighborModel/Source/NearestNeighborModel/Public/NearestNeighborModelInputInfo.h
@@ -5,8 +5,6 @@
 #include "MLDeformerMorphModelInputInfo.h"
 #include "NearestNeighborModelInputInfo.generated.h"
 
-class UNearestNeighborModelInstance;
-class UNearestNeighborModel;
 namespace UE::NearestNeighborModel
 {
     class FNearestNeighborEditorModel;
@@ -24,13 +22,12 @@ public:
     virtual void ExtractBoneRotations(USkeletalMeshComponent* SkelMeshComponent, TArray<float>& OutRotations) const override;
     // ~END UMLDeformerInputInfo overrides
 
-    friend class UNearestNeighborModel;
-    friend class UNearestNeighborModelInstance;
-    friend class UE::NearestNeighborModel::FNearestNeighborEditorModel;
+    void ComputeNetworkInput(TConstArrayView<FQuat> BoneRotations, TArrayView<float> OutputView) const;
 
 private:
-    void InitRefBoneRotations(USkeletalMesh* SkelMesh);
-    void ComputeNetworkInput(const TArray<FQuat>& BoneRotations, float* OutputBuffer, int64 StartIndex = 0) const;
+    friend class UE::NearestNeighborModel::FNearestNeighborEditorModel;
+
+    void InitRefBoneRotations(const USkeletalMesh* SkelMesh);
 
 private:
     UPROPERTY()
diff --git a/Engine/Plugins/Animation/MLDeformer/NearestNeighborModel/Source/NearestNeighborModel/Public/NearestNeighborModelInstance.h b/Engine/Plugins/Animation/MLDeformer/NearestNeighborModel/Source/NearestNeighborModel/Public/NearestNeighborModelInstance.h
--- a/Engine/Plugins/Animation/MLDeformer/NearestNeighborModel/Source/NearestNeighborModel/Public/NearestNeighborModelInstance.h
+++ b/Engine/Plugins/Animation/MLDeformer/NearestNeighborModel/Source/NearestNeighborModel/Public/NearestNeighborModelInstance.h
@@ -2,11 +2,13 @@
 #pragma once
 
 #include "MLDeformerMorphModelInstance.h"
+#include "Misc/Optional.h"
 #include "NearestNeighborModelInstance.generated.h"
 
-
+class UNearestNeighborModel;
 class UNearestNeighborModelInputInfo;
 class UNearestNeighborModelInstance;
+class UNearestNeighborOptimizedNetwork;
 class UNearestNeighborOptimizedNetworkInstance;
 namespace UE::NearestNeighborModel
 {
@@ -27,6 +29,7 @@ public:
     virtual bool SetupInputs() override;
     virtual FString CheckCompatibility(USkeletalMeshComponent* InSkelMeshComponent, bool LogIssues=false) override;
     virtual void Tick(float DeltaTime, float ModelWeight) override;
+    virtual int64 SetBoneTransforms(float* OutputBuffer, int64 OutputBufferSize, int64 StartIndex) override;
     // ~END UMLDeformerModelInstance overrides
 
     friend class UNearestNeighborModelInputInfo;
@@ -35,33 +38,34 @@ public:
     friend class UE::NearestNeighborModel::FNearestNeighborEditorModel;
     friend class UE::NearestNeighborModel::FNearestNeighborEditorModelActor;
 
-private:
-#if WITH_EDITORONLY_DATA
-    const TArray<uint32>& GetNearestNeighborIds() const { return NearestNeighborIds; }
-    uint32 NearestNeighborId(int32 PartId) const { return NearestNeighborIds[PartId]; }
-    int32 NeighborIdNum() const { return NearestNeighborIds.Num(); }
+#if WITH_EDITOR
+    TArray<uint32> GetNearestNeighborIds() const;
 #endif
 
-    void InitPreviousWeights();
-    void InitOptimizedNetworkInstance();
-    void GetInputDataPointer(float*& OutInputData, int32& OutNumInputFloats);
-    void GetOutputDataPointer(float*& OutOutputData, int32& OutNumOutputFloats);
+private:
+    using UNetworkPtr = TWeakObjectPtr<UNearestNeighborOptimizedNetwork>;
+    using UConstNetworkPtr = const TWeakObjectPtr<const UNearestNeighborOptimizedNetwork>;
 
-    virtual int64 SetBoneTransforms(float* OutputBuffer, int64 OutputBufferSize, int64 StartIndex) override;
+    void InitInstanceData(int32 NumMorphWeights = INDEX_NONE);
+
+    UNearestNeighborModel* GetCastModel() const;
+    void InitOptimizedNetworkInstance();
+    TOptional<TArrayView<float>> GetInputView() const;
+    TOptional<TArrayView<float>> GetOutputView() const;
 
     void RunNearestNeighborModel(float DeltaTime, float ModelWeight);
-    int32 FindNearestNeighbor(const float* PCAData, int32 PartId);
-    void UpdateWeightWithDecay(TArray<float>& MorphWeights, int32 Index, float W, float DecayCoeff);
 
+    /** This is the slow version of network inference. It is used by python. */
     UFUNCTION(BlueprintCallable, Category = "Python")
     TArray<float> Eval(const TArray<float>& InputData);
 
-#if WITH_EDITORONLY_DATA
+#if WITH_EDITOR
     TArray<uint32> NearestNeighborIds;
 #endif
 
     TArray<float> PreviousWeights;
+    TArray<float> DistanceBuffer;
 
-    UPROPERTY(Transient)
-    TObjectPtr<UNearestNeighborOptimizedNetworkInstance> OptimizedNetworkInstance = nullptr;
+    UPROPERTY()
+    TObjectPtr<UNearestNeighborOptimizedNetworkInstance> OptimizedNetworkInstance;
 };
diff --git a/Engine/Plugins/Animation/MLDeformer/NearestNeighborModel/Source/NearestNeighborModel/Public/NearestNeighborModelVizSettings.h b/Engine/Plugins/Animation/MLDeformer/NearestNeighborModel/Source/NearestNeighborModel/Public/NearestNeighborModelVizSettings.h
--- a/Engine/Plugins/Animation/MLDeformer/NearestNeighborModel/Source/NearestNeighborModel/Public/NearestNeighborModelVizSettings.h
+++ b/Engine/Plugins/Animation/MLDeformer/NearestNeighborModel/Source/NearestNeighborModel/Public/NearestNeighborModelVizSettings.h
@@ -2,14 +2,13 @@
 
 #pragma once
 
-#include "CoreTypes.h"
 #include "MLDeformerMorphModelVizSettings.h"
 #include "NearestNeighborModelVizSettings.generated.h"
 
 class UGeometryCache;
 namespace UE::NearestNeighborModel
 {
-    class FNearestNeighborEditorModel;
+	class FNearestNeighborModelVizSettingsDetails;
 };
 
 /**
@@ -24,23 +23,24 @@ class NEARESTNEIGHBORMODEL_API UNearestNeighborModelVizSettings
 public:
 	static FName GetNearestNeighborActorsOffsetPropertyName() { return GET_MEMBER_NAME_CHECKED(UNearestNeighborModelVizSettings, NearestNeighborActorsOffset); }
 	static FName GetNearestNeighborIdsPropertyName() { return GET_MEMBER_NAME_CHECKED(UNearestNeighborModelVizSettings, NearestNeighborIds); }
-	static FName GetNeighborStatsPartIdPropertyName() { return GET_MEMBER_NAME_CHECKED(UNearestNeighborModelVizSettings, NeighborStatsPartId); }
 	
-	friend class UE::NearestNeighborModel::FNearestNeighborEditorModel;
+	/** Whether to show verts */
+	UPROPERTY(EditAnywhere, Category = "Training Meshes")
+	bool bDrawVerts = false;
 
-protected:
-	UPROPERTY(EditAnywhere, Category = "Live Settings")
-	float NearestNeighborActorsOffset = 2.0f;
+	/** Show vertices in this section */
+	UPROPERTY(EditAnywhere, Category = "Training Meshes", Meta = (DisplayName = "Show Verts in", EditorCondition = "bDrawVerts"))
+	int32 VertVizSectionIndex = INDEX_NONE;
 
+	/** The section used to display the nearest neighbor. */	
+	UPROPERTY(EditAnywhere, Category = "Live Settings", Meta = (DisplayName = "Actor Section Index"))
+	int32 NearestNeighborActorSectionIndex = 0;
+	
+	/** The offset of the nearest neighbor actor from the mesh. */
+	UPROPERTY(EditAnywhere, Category = "Live Settings", Meta = (DisplayName = "Actor Offset"))
+	float NearestNeighborActorsOffset = 2.0f;
+	
 	UPROPERTY(VisibleAnywhere, Category = "Live Settings")
-	TArray<uint32> NearestNeighborIds;
-
-	UPROPERTY(EditAnywhere, Category = "Live Settings", meta = (DisplayName = "Part Id"))
-	int32 NeighborStatsPartId = 0;
-
-private:
-	void SetNearestNeighborActorsOffset(float InOffset) { NearestNeighborActorsOffset = InOffset; }
-	float GetNearestNeighborActorsOffset() const { return NearestNeighborActorsOffset; }
-	void SetNearestNeighborIds(const TArray<uint32>& InNearestNeighborIds) { NearestNeighborIds = InNearestNeighborIds; }
+	TArray<int32> NearestNeighborIds;
 #endif
 };
diff --git a/Engine/Plugins/Animation/MLDeformer/NearestNeighborModel/Source/NearestNeighborModel/Public/NearestNeighborOptimizedNetworkLoader.h b/Engine/Plugins/Animation/MLDeformer/NearestNeighborModel/Source/NearestNeighborModel/Public/NearestNeighborOptimizedNetworkLoader.h
--- a/Engine/Plugins/Animation/MLDeformer/NearestNeighborModel/Source/NearestNeighborModel/Public/NearestNeighborOptimizedNetworkLoader.h
+++ b/Engine/Plugins/Animation/MLDeformer/NearestNeighborModel/Source/NearestNeighborModel/Public/NearestNeighborOptimizedNetworkLoader.h
@@ -11,7 +11,7 @@ class UNearestNeighborOptimizedNetwork;
  *  LoadOptimizedNetwork is implemented in python.
  */
 UCLASS()
-class NEARESTNEIGHBORMODEL_API UNearestNeighborOptimizedNetworkLoader
+class UE_DEPRECATED(5.4, "UNearestNeighborOptimizedNetworkLoader is deprecated. Use NNE instead.")  NEARESTNEIGHBORMODEL_API UNearestNeighborOptimizedNetworkLoader
 	: public UObject
 {
 	GENERATED_BODY()
diff --git a/Engine/Plugins/Animation/MLDeformer/NearestNeighborModel/Source/NearestNeighborModelEditor/Private/NearestNeighborEditorHelpers.h b/Engine/Plugins/Animation/MLDeformer/NearestNeighborModel/Source/NearestNeighborModelEditor/Private/NearestNeighborEditorHelpers.h
new file mode 100644
--- /dev/null
+++ b/Engine/Plugins/Animation/MLDeformer/NearestNeighborModel/Source/NearestNeighborModelEditor/Private/NearestNeighborEditorHelpers.h
@@ -0,0 +1,32 @@
+// Copyright Epic Games, Inc. All Rights Reserved.
+
+#pragma once
+
+#include "CoreTypes.h"
+
+#include "NearestNeighborEditorHelpers.generated.h"
+
+class UAnimSequence;
+UCLASS(Blueprintable)
+class NEARESTNEIGHBORMODELEDITOR_API UNearestNeighborAnimStream : public UObject
+{
+	GENERATED_BODY()
+public:
+	UFUNCTION(BlueprintCallable, Category = "Python")
+	void Init(USkeleton* InSkeleton);
+
+	UFUNCTION(BlueprintPure, Category = "Python")
+	bool IsValid() const;
+
+	UFUNCTION(BlueprintCallable, Category = "Python")
+	bool AppendFrames(const UAnimSequence* Anim, TArray<int32> Frames);
+
+	UFUNCTION(BlueprintPure, Category = "Python")
+	bool ToAnim(UAnimSequence* OutAnim) const;
+
+private:
+	TObjectPtr<USkeleton> Skeleton;
+	TArray<TArray<FVector3f>> PosKeys;
+	TArray<TArray<FQuat4f>> RotKeys;
+	TArray<TArray<FVector3f>> ScaleKeys;
+};
\ No newline at end of file
