Nearestneighbormodel: refactor

diff --git a/Engine/Plugins/Animation/MLDeformer/NearestNeighborModel/Source/NearestNeighborModelEditor/Private/NearestNeighborEditorModel.cpp b/Engine/Plugins/Animation/MLDeformer/NearestNeighborModel/Source/NearestNeighborModelEditor/Private/NearestNeighborEditorModel.cpp
--- a/Engine/Plugins/Animation/MLDeformer/NearestNeighborModel/Source/NearestNeighborModelEditor/Private/NearestNeighborEditorModel.cpp
+++ b/Engine/Plugins/Animation/MLDeformer/NearestNeighborModel/Source/NearestNeighborModelEditor/Private/NearestNeighborEditorModel.cpp
@@ -1,38 +1,56 @@
 // Copyright Epic Games, Inc. All Rights Reserved.
 
 #include "NearestNeighborEditorModel.h"
-#include "IDetailsView.h"
-#include "NearestNeighborModel.h"
-#include "NearestNeighborModelInstance.h"
-#include "NearestNeighborTrainingModel.h"
-#include "NearestNeighborModelInputInfo.h"
-#include "NearestNeighborEditorModelActor.h"
-#include "NearestNeighborModelStyle.h"
-#include "NearestNeighborGeomCacheSampler.h"
+
+#include "Algo/Copy.h"
+#include "Animation/AnimSequence.h"
 #include "Components/ExternalMorphSet.h"
-#include "MLDeformerComponent.h"
-#include "MLDeformerEditorToolkit.h"
-#include "MLDeformerAsset.h"
+#include "Engine/SkeletalMesh.h"
 #include "GeometryCache.h"
 #include "GeometryCacheComponent.h"
-#include "Animation/DebugSkelMeshComponent.h"
-#include "Animation/MorphTarget.h"
-#include "Animation/AnimSequence.h"
-#include "Animation/AnimInstance.h"
-#include "BonePose.h"
-#include "PackageTools.h"
-#include "ObjectTools.h"
-#include "UObject/SavePackage.h"
-#include "Misc/FileHelper.h"
-
+#include "IPersonaPreviewScene.h"
+#include "Misc/MessageDialog.h"
+#include "MLDeformerAsset.h"
+#include "MLDeformerComponent.h"
+#include "MLDeformerEditorStyle.h"
+#include "NearestNeighborGeomCacheSampler.h"
+#include "NearestNeighborModel.h"
+#include "NearestNeighborModelInputInfo.h"
+#include "NearestNeighborModelInstance.h"
+#include "NearestNeighborModelStyle.h"
+#include "NearestNeighborModelVizSettings.h"
+#include "NearestNeighborTrainingModel.h"
+#include "Rendering/SkeletalMeshModel.h"
+#include "SceneManagement.h"
 
 #define LOCTEXT_NAMESPACE "NearestNeighborEditorModel"
 
 namespace UE::NearestNeighborModel
 {
-	using namespace UE::MLDeformer;
+	namespace Private
+	{
+		void AddFloatArrayToDeltaArray(TConstArrayView<float> FloatArr, TConstArrayView<int32> VertexMap, TArrayView<FVector3f> OutDeltaArr)
+		{
+			const int32 SectionNumVerts = VertexMap.Num();
+			check(FloatArr.Num() == SectionNumVerts * 3);
+			for (int32 Index = 0; Index < SectionNumVerts; ++Index)
+			{
+				const int32 FloatIndex = Index * 3;
+				const int32 DeltaIndex = VertexMap[Index];
+				OutDeltaArr[DeltaIndex] = FVector3f(FloatArr[FloatIndex], FloatArr[FloatIndex + 1], FloatArr[FloatIndex + 2]);
+			}
+		}
+	};
 
-	FMLDeformerEditorModel* FNearestNeighborEditorModel::MakeInstance()
+	FNearestNeighborEditorModel::~FNearestNeighborEditorModel()
+	{
+		if (NearestNeighborActor.Get())
+		{
+			EditorActors.Remove(NearestNeighborActor.Get());
+		}
+	}
+
+	UE::MLDeformer::FMLDeformerEditorModel* FNearestNeighborEditorModel::MakeInstance()
 	{
 		return new FNearestNeighborEditorModel();
 	}
@@ -41,975 +59,854 @@ namespace UE::NearestNeighborModel
 	{
 		FMLDeformerGeomCacheEditorModel::Init(InitSettings);
 		InitInputInfo(Model->GetInputInfo());
+		VertexMapSelector = MakeUnique<FVertexMapSelector>();
+		VertexMapSelector->Update(Model->GetSkeletalMesh());
+		const UNearestNeighborModel* const NearestNeighborModel = GetCastModel();
+		if (!NearestNeighborModel)
+		{
+			return;
+		}
+		const int32 NumSections = NearestNeighborModel->GetNumSections();
+		if (UNearestNeighborModelVizSettings* const NNViz = GetCastVizSettings())
+		{
+			VertVizSelector = MakeUnique<FVertVizSelector>(NNViz);
+			VertVizSelector->Update(NumSections);
+		}
 	}
 
-	TSharedPtr<FMLDeformerSampler> FNearestNeighborEditorModel::CreateSamplerObject() const
+	TSharedPtr<UE::MLDeformer::FMLDeformerSampler> FNearestNeighborEditorModel::CreateSamplerObject() const
 	{
 		return MakeShared<FNearestNeighborGeomCacheSampler>();
 	}
 
-	void FNearestNeighborEditorModel::OnPropertyChanged(FPropertyChangedEvent& PropertyChangedEvent)
+	void FNearestNeighborEditorModel::CreateActors(const TSharedRef<IPersonaPreviewScene>& InPersonaPreviewScene)
 	{
-		const FProperty* Property = PropertyChangedEvent.Property;
-		if (Property == nullptr)
+		FMLDeformerMorphModelEditorModel::CreateActors(InPersonaPreviewScene);
+	
+		UWorld* const World = InPersonaPreviewScene->GetWorld();
+		if (!World)
 		{
 			return;
 		}
 
-		FMLDeformerMorphModelEditorModel::OnPropertyChanged(PropertyChangedEvent);
-
-		if (Property->GetFName() == UMLDeformerModel::GetSkeletalMeshPropertyName() ||
-			Property->GetFName() == UNearestNeighborModel::GetClothPartEditorDataPropertyName() ||
-		    Property->GetFName() == UMLDeformerGeomCacheModel::GetTrainingInputAnimsPropertyName() ||
-		    PropertyChangedEvent.GetMemberPropertyName() == TEXT("TrainingInputAnims") ||
-			PropertyChangedEvent.GetMemberPropertyName() == UNearestNeighborModel::GetClothPartEditorDataPropertyName())
+		if (NearestNeighborActor.Get())
 		{
-			GetNearestNeighborModel()->InvalidateClothPartData();
-			GetEditor()->GetModelDetailsView()->ForceRefresh();
-		}
-
-		if (Property->GetFName() == UNearestNeighborModel::GetNearestNeighborDataPropertyName() || (PropertyChangedEvent.MemberProperty != nullptr && PropertyChangedEvent.MemberProperty->GetFName() == UNearestNeighborModel::GetNearestNeighborDataPropertyName()))
-		{
-			GetNearestNeighborModel()->InvalidateNearestNeighborData();
-			GetEditor()->GetModelDetailsView()->ForceRefresh();
+			EditorActors.Remove(NearestNeighborActor.Get());
+			if (NearestNeighborActor->GetActor())
+			{
+				World->DestroyActor(NearestNeighborActor->GetActor(), true);
+			}
 		}
+		NearestNeighborActor = CreateNearestNeighborActor(World);
+		UpdateNearestNeighborActor(*NearestNeighborActor);
+		EditorActors.Add(NearestNeighborActor.Get());
+	}
 
-		if (Property->GetFName() == UNearestNeighborModel::GetMorphCompressionLevelPropertyName() || Property->GetFName() == UNearestNeighborModel::GetMorphDeltaZeroThresholdPropertyName())
+	void FNearestNeighborEditorModel::Tick(FEditorViewportClient* ViewportClient, float DeltaTime)
+	{
+		FMLDeformerEditorModel::Tick(ViewportClient, DeltaTime);
+		if (NearestNeighborActor.Get())
 		{
-			GetNearestNeighborModel()->InvalidateMorphTargetData();
-			GetEditor()->GetModelDetailsView()->ForceRefresh();
+			UpdateNearestNeighborActor(*NearestNeighborActor);
+			NearestNeighborActor->Tick();
 		}
+		UpdateNearestNeighborIds();
 	}
 
-	void FNearestNeighborEditorModel::OnPostTraining(ETrainingResult TrainingResult, bool bUsePartiallyTrainedWhenAborted)
+	void FNearestNeighborEditorModel::InitInputInfo(UMLDeformerInputInfo* InputInfo)
+	{
+		FMLDeformerEditorModel::InitInputInfo(InputInfo);
+		UNearestNeighborModelInputInfo* NearestNeighborInputInfo = static_cast<UNearestNeighborModelInputInfo*>(InputInfo);
+		NearestNeighborInputInfo->InitRefBoneRotations(Model->GetSkeletalMesh());
+	}
+	
+	ETrainingResult FNearestNeighborEditorModel::Train()
 	{
-		FMLDeformerEditorModel::OnPostTraining(TrainingResult, bUsePartiallyTrainedWhenAborted);
-		if (TrainingResult == ETrainingResult::Success || (TrainingResult == ETrainingResult::Aborted && bUsePartiallyTrainedWhenAborted))
+		UNearestNeighborModel* const NearestNeighborModel = GetCastModel();
+		if (!NearestNeighborModel)
+		{
+			return ETrainingResult::FailPythonError;
+		}
+		if(OpFlag::HasError(NearestNeighborModel->UpdateForTraining()))
 		{
-			GetNearestNeighborModel()->InvalidateNearestNeighborData();
-			InitTestMLDeformerPreviousWeights();
-			GetEditor()->GetModelDetailsView()->ForceRefresh();
+			return ETrainingResult::FailPythonError;
 		}
+		if (!NearestNeighborModel->IsReadyForTraining())
+		{
+			return ETrainingResult::FailPythonError;
+		}
+		const int32 NumFrames = GetNumFramesForTraining();
+		const int32 NumPCACoeffs = NearestNeighborModel->GetTotalNumPCACoeffs();
+		if (NumFrames <= NumPCACoeffs)
+		{
+			UE_LOG(LogNearestNeighborModel, Error, TEXT("Training frames (%d) must be greater than the number of PCA coefficients (%d)"), NumFrames, NumPCACoeffs);
+			return ETrainingResult::FailPythonError;
+		}
+		return TrainModel<UNearestNeighborTrainingModel>(this);
 	}
-	
+
 	bool FNearestNeighborEditorModel::LoadTrainedNetwork() const
 	{
-		UNearestNeighborModel* NearestNeighborModel = GetNearestNeighborModel();
+		UNearestNeighborModel* const NearestNeighborModel = GetCastModel();
 		if (NearestNeighborModel)
 		{	
-			const FString File = GetNearestNeighborModel()->GetModelDir() / TEXT("NearestNeighborModel.ubnne");
-			if (NearestNeighborModel->ShouldUseOptimizedNetwork())
+			const FString File = NearestNeighborModel->GetModelDir() / TEXT("NearestNeighborModel.ubnne");
+			const bool bSuccess = NearestNeighborModel->LoadOptimizedNetworkFromFile(File);
+			if (bSuccess)
 			{
-				const bool bSuccess = NearestNeighborModel->LoadOptimizedNetworkFromFile(File);
-				if (bSuccess)
+				UNearestNeighborModelInstance* ModelInstance = GetTestNearestNeighborModelInstance();
+				if (ModelInstance)
 				{
-					UNearestNeighborModelInstance* ModelInstance = static_cast<UNearestNeighborModelInstance*>(GetTestMLDeformerModelInstance());
-					if (ModelInstance)
-					{
-						ModelInstance->InitOptimizedNetworkInstance();
-						NearestNeighborModel->SetUseOptimizedNetwork(true);
-						return true;
-					}
+					ModelInstance->InitOptimizedNetworkInstance();
+					return true;
 				}
 			}
 		}
 		return false;
 	}
 
-	FString FNearestNeighborEditorModel::GetTrainedNetworkOnnxFile() const
+	void FNearestNeighborEditorModel::OnPostTraining(ETrainingResult TrainingResult, bool bUsePartiallyTrainedWhenAborted)
 	{
-		const UNearestNeighborModel* NearestNeighborModel = GetNearestNeighborModel();
-		if (NearestNeighborModel)
+		if (TrainingResult == ETrainingResult::Aborted && !bUsePartiallyTrainedWhenAborted)
 		{
-			return NearestNeighborModel->GetModelDir() + TEXT("/NearestNeighborModel.onnx");
+			GetMorphModel()->SetMorphTargetDeltas(MorphTargetDeltasBackup);
+			GetMorphModel()->SetMorphTargetsMinMaxWeights(MorphTargetsMinMaxWeightsBackup);
 		}
-		else
+		else if (TrainingResult == ETrainingResult::Success || (TrainingResult == ETrainingResult::Aborted && bUsePartiallyTrainedWhenAborted))
 		{
-			return FMLDeformerEditorModel::GetTrainedNetworkOnnxFile();
+			UNearestNeighborModel* const NearestNeighborModel = GetCastModel();
+			if (!NearestNeighborModel)
+			{
+				return;
+			}
+			NearestNeighborModel->InvalidateInference();
+			if (NearestNeighborModel->DoesUseFileCache())
+			{
+				NearestNeighborModel->UpdateFileCache();
+			}
+			ResetMorphTargets();
 		}
+		FMLDeformerGeomCacheEditorModel::OnPostTraining(TrainingResult, bUsePartiallyTrainedWhenAborted);
 	}
 
-	int32 FNearestNeighborEditorModel::GetNumTrainingFrames() const
+	FMLDeformerTrainingInputAnim* FNearestNeighborEditorModel::GetTrainingInputAnim(int32 Index) const
 	{
-		if (NumTrainingFramesOverride >= 0)
+		UNearestNeighborModel* const NearestNeighborModel = GetCastModel();
+		if (!NearestNeighborModel)
 		{
-			return NumTrainingFramesOverride;
+			return nullptr;
+		}
+		TArray<FMLDeformerGeomCacheTrainingInputAnim>& TrainingInputAnims = NearestNeighborModel->GetTrainingInputAnims();
+		const int32 NumTrainingAnims = TrainingInputAnims.Num();
+		const int32 NumSections = NearestNeighborModel->GetNumSections();
+		if (Index >= 0 && Index < NumTrainingAnims)
+		{
+			return &TrainingInputAnims[Index];
+		}
+		else if (Index >= NumTrainingAnims && Index < NumTrainingAnims + NumSections)
+		{
+			const int32 SectionIndex = Index - NumTrainingAnims;
+			return NearestNeighborModel->GetNearestNeighborAnim(SectionIndex);
 		}
 		else
 		{
-			return FMLDeformerMorphModelEditorModel::GetNumTrainingFrames();
+			return nullptr;
 		}
 	}
 
-	UNearestNeighborModelVizSettings* FNearestNeighborEditorModel::GetNearestNeighborModelVizSettings() const
-	{
-		return Cast<UNearestNeighborModelVizSettings>(GetMorphModel()->GetVizSettings()); 
-	}
-
-	void FNearestNeighborEditorModel::CreateActors(const TSharedRef<IPersonaPreviewScene>& InPersonaPreviewScene)
+	void FNearestNeighborEditorModel::UpdateTimelineTrainingAnimList()
 	{
-		FMLDeformerMorphModelEditorModel::CreateActors(InPersonaPreviewScene);
-
-		UWorld* World = InPersonaPreviewScene->GetWorld();
-		CreateNearestNeighborActors(World);
-		EditorWorld = World;
-	}
-
-	void FNearestNeighborEditorModel::Tick(FEditorViewportClient* ViewportClient, float DeltaTime)
-	{
-		FMLDeformerEditorModel::Tick(ViewportClient, DeltaTime);
-		if (!NearestNeighborActors.IsEmpty())
+		TArray<TSharedPtr<FMLDeformerTrainingInputAnimName>> NameList;
+		
+		// Build the list of names, based on the training inputs.
+		const int32 NumAnims = GetNumTrainingInputAnims();
+		for (int32 AnimIndex = 0; AnimIndex < NumAnims; ++AnimIndex)
 		{
-			UNearestNeighborModelVizSettings* NNViz = GetNearestNeighborModelVizSettings();
-			const float Offset = NNViz->GetNearestNeighborActorsOffset();
-			for (FNearestNeighborEditorModelActor* NearestNeighborActor : NearestNeighborActors)
+			const FMLDeformerTrainingInputAnim* Anim = GetTrainingInputAnim(AnimIndex);
+			if (Anim)
 			{
-				if (NearestNeighborActor)
+				if (Anim->IsValid())
 				{
-					NearestNeighborActor->TickNearestNeighborActor();
-					NearestNeighborActor->SetMeshOffsetFactor(Offset);
+					TSharedPtr<FMLDeformerTrainingInputAnimName> AnimName = MakeShared<FMLDeformerTrainingInputAnimName>();
+					AnimName->TrainingInputAnimIndex = AnimIndex;
+					AnimName->Name = FString::Printf(TEXT("[#%d] %s"), AnimIndex, *Anim->GetAnimSequence()->GetName());
+					NameList.Emplace(AnimName);
 				}
 			}
-			const UNearestNeighborModelInstance* ModelInstance = static_cast<UNearestNeighborModelInstance*>(GetTestMLDeformerModelInstance());
-			UNearestNeighborModel* NearestNeighborModel = GetNearestNeighborModel();
-			if (ModelInstance && NearestNeighborModel)
+		}
+
+		if (const UNearestNeighborModel* const NearestNeighborModel = GetCastModel())
+		{
+			const int32 NumSections = NearestNeighborModel->GetNumSections();
+			for (int32 Index = 0; Index < NumSections; ++Index)
 			{
-				const TArray<uint32>& NeighborIds = ModelInstance->GetNearestNeighborIds();
-				TArray<uint32> AssetNeighborIds; 
-				const int32 NumParts = FMath::Min(NeighborIds.Num(), NearestNeighborModel->GetNumParts());
-				AssetNeighborIds.SetNum(NumParts);
-				for(int32 Index = 0; Index < NumParts; Index++)
+				const FMLDeformerTrainingInputAnim* Anim = NearestNeighborModel->GetNearestNeighborAnim(Index);
+				if (Anim && Anim->IsValid())
 				{
-					const int32 Id = NeighborIds[Index];
-					const TArray<int32>& IndexMap = NearestNeighborModel->ClothPartData[Index].AssetNeighborIndexMap;
-					AssetNeighborIds[Index] = IndexMap.IsValidIndex(Id) ? IndexMap[Id] : INDEX_NONE;
+					TSharedPtr<FMLDeformerTrainingInputAnimName> AnimName = MakeShared<FMLDeformerTrainingInputAnimName>();
+					AnimName->TrainingInputAnimIndex = Index + NumAnims;
+					AnimName->Name = FString::Printf(TEXT("Neighbor[#%d] %s"), Index, *Anim->GetAnimSequence()->GetName());	
+					NameList.Emplace(AnimName);
 				}
-				NNViz->SetNearestNeighborIds(ModelInstance->GetNearestNeighborIds());
 			}
 		}
+
+		SetTimelineAnimNames(NameList);
 	}
 
-	void FNearestNeighborEditorModel::CreateNearestNeighborActors(UWorld* World, int32 StartIndex)
+	void FNearestNeighborEditorModel::OnPropertyChanged(FPropertyChangedEvent& PropertyChangedEvent)
 	{
-		UNearestNeighborModel* NearestNeighborModel = GetNearestNeighborModel();
-		const int32 NumParts = NearestNeighborModel->GetNumParts();
-
-		if (StartIndex == 0)
+		FMLDeformerMorphModelEditorModel::OnPropertyChanged(PropertyChangedEvent);
+		const FProperty* Property = PropertyChangedEvent.Property;
+		if (Property == nullptr)
 		{
-			NearestNeighborActors.Reset();
+			return;
 		}
-		NearestNeighborActors.SetNumZeroed(NumParts);
-
-		const UNearestNeighborModelVizSettings* NNViz = GetNearestNeighborModelVizSettings();
-		const float Offset = NNViz->GetNearestNeighborActorsOffset();
-		
-		for(int32 PartId = StartIndex; PartId < NumParts; PartId++)
+		UNearestNeighborModel* const NearestNeighborModel = GetCastModel();
+		if (!NearestNeighborModel)
 		{
-			UGeometryCache* GeomCache = NearestNeighborModel->GetNearestNeighborCache(PartId);
-			const FLinearColor LabelColor = FNearestNeighborModelEditorStyle::Get().GetColor("NearestNeighborModel.NearestNeighborActors.LabelColor");
-			const FLinearColor WireframeColor = FNearestNeighborModelEditorStyle::Get().GetColor("NearestNeighborModel.NearestNeighborActors.WireframeColor");
-			CreateGeomCacheActor(
-				World, 
-				ActorID_NearestNeighborActors,
-				*FString::Printf(TEXT("NearestNeighbors%d"), PartId), 
-				GeomCache, 
-				LabelColor,					
-				WireframeColor,
-				LOCTEXT("TestNearestNeighborLabelText", "Nearest Neigbors"),
-				false);
-			FNearestNeighborEditorModelActor* NearestNeighborActor = static_cast<FNearestNeighborEditorModelActor*>(EditorActors.Last());
-			UMLDeformerComponent* MLDeformerComponent = GetTestMLDeformerComponent();
-			NearestNeighborActor->InitNearestNeighborActor(PartId, MLDeformerComponent);
-			NearestNeighborActor->SetMeshOffsetFactor(Offset * 2.0f);
-			NearestNeighborActors[PartId] = NearestNeighborActor;
+			return;
 		}
-	}
 
-	FMLDeformerEditorActor* FNearestNeighborEditorModel::CreateEditorActor(const FMLDeformerEditorActor::FConstructSettings& Settings) const
-	{
-		return new FNearestNeighborEditorModelActor(Settings);
-	}
-
-	void FNearestNeighborEditorModel::InitInputInfo(UMLDeformerInputInfo* InputInfo)
-	{
-		FMLDeformerEditorModel::InitInputInfo(InputInfo);
-		UNearestNeighborModelInputInfo* NearestNeighborInputInfo = static_cast<UNearestNeighborModelInputInfo*>(InputInfo);
-		NearestNeighborInputInfo->InitRefBoneRotations(Model->GetSkeletalMesh());
-	}
-
-	ETrainingResult FNearestNeighborEditorModel::Train()
-	{
-		UNearestNeighborModel* NearestNeighborModel = GetNearestNeighborModel();
-		MorphTargetUpdateResult = EUpdateResult::SUCCESS;
-		if (!NearestNeighborModel->IsClothPartDataValid())
+		if (Property->GetFName() == UMLDeformerModel::GetSkeletalMeshPropertyName())
 		{
-			MorphTargetUpdateResult |= NearestNeighborModel->UpdateClothPartData();
-			if (HasError(MorphTargetUpdateResult))
+			if (VertexMapSelector.IsValid())
 			{
-				UpdateNearestNeighborActors();
-				return ETrainingResult::FailOnData;
+				VertexMapSelector->Update(NearestNeighborModel->GetSkeletalMesh());
 			}
 		}
 
-		return TrainModel<UNearestNeighborTrainingModel>(this);
-	}
-
-	uint8 FNearestNeighborEditorModel::SetSamplerPartData(const int32 PartId)
-	{
-		FMLDeformerSampler* Sampler = GetNumTrainingInputAnims() > 0 ? GetSamplerForTrainingAnim(0) : nullptr;
-		FNearestNeighborGeomCacheSampler* GeomCacheSampler = static_cast<FNearestNeighborGeomCacheSampler*>(Sampler);
-		UNearestNeighborModel* NearestNeighborModel = GetNearestNeighborModel();
-		if (GeomCacheSampler && NearestNeighborModel && PartId < NearestNeighborModel->GetNumParts())
+		if (Property->GetFName() == UNearestNeighborModel::GetSectionsPropertyName())
 		{
-			TObjectPtr<USkeletalMeshComponent> SkeletalMeshComponent = GeomCacheSampler->GetSkeletalMeshComponent();
-			TObjectPtr<UGeometryCacheComponent> GeometryCacheComponent = GeomCacheSampler->GetGeometryCacheComponent();
-
-			const TObjectPtr<UAnimSequence> AnimSequence = NearestNeighborModel->GetNearestNeighborSkeletons(PartId);
-			const TObjectPtr<UGeometryCache> GeometryCache = NearestNeighborModel->GetNearestNeighborCache(PartId);
-
-			if (SkeletalMeshComponent && GeometryCacheComponent && AnimSequence && GeometryCache)
-			{
-				const int32 NumNeighborsFromGeomCache = NearestNeighborModel->GetNumNeighborsFromGeometryCache(PartId);
-				const int32 NumNeighborsFromAnimSequence = NearestNeighborModel->GetNumNeighborsFromAnimSequence(PartId);
-				const int32 NumFrames = FMath::Min(NumNeighborsFromGeomCache, NumNeighborsFromAnimSequence);
-				if (NumFrames > 0)
-				{
-					AnimSequence->Interpolation = EAnimInterpolationType::Step;
-
-					SkeletalMeshComponent->SetAnimationMode(EAnimationMode::AnimationSingleNode);
-					SkeletalMeshComponent->SetAnimation(AnimSequence);
-					SkeletalMeshComponent->SetPosition(0.0f);
-					SkeletalMeshComponent->SetPlayRate(1.0f);
-					SkeletalMeshComponent->Play(false);
-					SkeletalMeshComponent->RefreshBoneTransforms();
-					if (SkeletalMeshComponent->GetAnimInstance())
-					{
-						SkeletalMeshComponent->GetAnimInstance()->GetRequiredBones().SetUseRAWData(true);
-					}
-
-					// assuming MeshMappings do not change
-					GeometryCacheComponent->SetGeometryCache(GeometryCache);
-					GeometryCacheComponent->ResetAnimationTime();
-					GeometryCacheComponent->SetLooping(false);
-					GeometryCacheComponent->SetManualTick(true);
-					GeometryCacheComponent->SetPlaybackSpeed(1.0f);
-					GeometryCacheComponent->Play();
-					if (uint8 ReturnCode = GeomCacheSampler->GenerateMeshMappings(); HasError(ReturnCode))
-					{
-						return ReturnCode;
-					}
-
-					NumTrainingFramesOverride = NumFrames;
-
-					if (NumNeighborsFromGeomCache == NumNeighborsFromAnimSequence)
-					{
-						return EUpdateResult::SUCCESS;
-					}
-					else
-					{
-						UE_LOG(LogNearestNeighborModel, Warning, TEXT("NearestNeighborData: part %d frame mismatch: AnimSequence has %d frames and GeometryCache has %d frames. Using %d frames only."), PartId, NumNeighborsFromAnimSequence, NumNeighborsFromGeomCache, NumFrames);
-						return EUpdateResult::WARNING;
-					}
-				}
-				else
-				{
-					UE_LOG(LogNearestNeighborModel, Warning, TEXT("Part %d: AnimSequence or GeometryCache has zero frames"), PartId);
-					return EUpdateResult::WARNING;
-				}
-			}
-			else if (AnimSequence == nullptr && GeometryCache == nullptr)
-			{
-				return EUpdateResult::SUCCESS;
-			}
-			else
+			if (VertVizSelector.IsValid())
 			{
-				UE_LOG(LogNearestNeighborModel, Warning, TEXT("Part %d is skipped because AnimSequence or GeometryCache is None"), PartId);
-				return EUpdateResult::WARNING;
+				VertVizSelector->Update(NearestNeighborModel->GetNumSections());
 			}
+			UpdateTimelineTrainingAnimList();
+		}
+
+		if (Property->GetFName() == UNearestNeighborModelSection::GetNeighborPosesPropertyName())
+		{
+			UpdateTimelineTrainingAnimList();
 		}
-		UE_LOG(LogNearestNeighborModel, Error, TEXT("SetSamplerPartData: unknown error"));
-		return EUpdateResult::ERROR;
 	}
 
-	int32 FNearestNeighborEditorModel::GetNumParts()
+	void FNearestNeighborEditorModel::Render(const FSceneView* View, FViewport* Viewport, FPrimitiveDrawInterface* PDI)
 	{
-		UNearestNeighborModel *NearestNeighborModel = GetNearestNeighborModel();
-		if (NearestNeighborModel != nullptr)
+		FMLDeformerGeomCacheEditorModel::Render(View, Viewport, PDI);
+		const FMLDeformerSampler* Sampler = GetSamplerForActiveAnim();
+		if (!Sampler || !Sampler->IsInitialized())
 		{
-			return NearestNeighborModel->GetNumParts();
+			return;
 		}
-		else
+
+		const UNearestNeighborModelVizSettings* VizSettings = GetCastVizSettings();
+		if (!VizSettings)
 		{
-			return 0;
+			return;
 		}
-	}
 
-	void FNearestNeighborEditorModel::ResetSamplerData()
-	{
-		FMLDeformerSampler* Sampler = GetNumTrainingInputAnims() > 0 ? GetSamplerForTrainingAnim(0) : nullptr;
-		FNearestNeighborGeomCacheSampler* GeomCacheSampler = static_cast<FNearestNeighborGeomCacheSampler*>(Sampler);
-		UNearestNeighborModel* NearestNeighborModel = GetNearestNeighborModel();
-		if (GeomCacheSampler && NearestNeighborModel)
+		if (VizSettings->GetVisualizationMode() == EMLDeformerVizMode::TrainingData)
 		{
-			TObjectPtr<USkeletalMeshComponent> SkeletalMeshComponent = GeomCacheSampler->GetSkeletalMeshComponent();
-			const TObjectPtr<UAnimSequence> AnimSequence = GetTrainingInputAnim(0)->GetAnimSequence();
+			const TArray<FVector3f>& LinearSkinnedPositions = Sampler->GetSkinnedVertexPositions();
+			bool bDrawVerts = VizSettings->bDrawVerts;
 
-			if (SkeletalMeshComponent && AnimSequence)
+			// Disable drawing deltas when we're playing the training anim sequence, as that can get too slow.
+			const UE::MLDeformer::FMLDeformerEditorActor* BaseActor = FindEditorActor(UE::MLDeformer::ActorID_Train_Base);
+			if (BaseActor)
 			{
-				SkeletalMeshComponent->SetAnimationMode(EAnimationMode::AnimationSingleNode);
-				SkeletalMeshComponent->SetAnimation(AnimSequence);
-				SkeletalMeshComponent->SetPosition(0.0f);
-				SkeletalMeshComponent->SetPlayRate(1.0f);
-				SkeletalMeshComponent->Play(false);
-				SkeletalMeshComponent->RefreshBoneTransforms();
+				bDrawVerts &= !BaseActor->IsPlaying();
 			}
-			NumTrainingFramesOverride = -1;
-		}
-	}
 
-	void FNearestNeighborEditorModel::UpdateNearestNeighborActors()
-	{
-		if (EditorWorld != nullptr)
-		{
-			const UNearestNeighborModel* NearestNeighborModel = GetNearestNeighborModel();
-			const int32 TargetNumActors = NearestNeighborModel->GetNumParts();
-			if (NearestNeighborActors.Num() > TargetNumActors)
+			// Draw the verts for the current frame.
+			if (bDrawVerts)
 			{
-				for (int32 i = NearestNeighborActors.Num() - 1; i >= TargetNumActors; i--)
+				const FLinearColor VertsColor = FNearestNeighborModelEditorStyle::Get().GetColor("NearestNeighborModel.Verts.VertsColor");
+				const uint8 DepthGroup = VizSettings->GetXRayDeltas() ? 100 : 0;
+				const float PointSize = 5;
+				const int32 SectionIndex = VizSettings->VertVizSectionIndex;
+				if (SectionIndex != INDEX_NONE)
 				{
-					FNearestNeighborEditorModelActor* EditorActor = NearestNeighborActors[i];
-					NearestNeighborActors.RemoveAt(i);
-					EditorActors.Remove(EditorActor);
-					EditorWorld->DestroyActor(EditorActor->GetActor(), true/*ShouldModifyLevel*/);
-					delete EditorActor;
+					const UNearestNeighborModel* const NearestNeighborModel = GetCastModel();
+					if (NearestNeighborModel)
+					{
+						TConstArrayView<int32> VertexMap = NearestNeighborModel->GetSection(SectionIndex).GetVertexMap();
+						for (int32 Index = 0; Index < VertexMap.Num(); ++Index)
+						{
+							const int32 ArrayIndex = 3 * VertexMap[Index];
+							const FVector VertexPos = (FVector)LinearSkinnedPositions[VertexMap[Index]];
+							PDI->DrawPoint(VertexPos, VertsColor, PointSize, DepthGroup);
+						}
+					}
+				}
+				else // draw all verts
+				{
+					for (int32 Index = 0; Index < LinearSkinnedPositions.Num(); ++Index)
+					{
+						const FVector VertexPos = (FVector)LinearSkinnedPositions[Index];
+						PDI->DrawPoint(VertexPos, VertsColor, PointSize, DepthGroup);
+					}
 				}
-			}
-			if (NearestNeighborActors.Num() < TargetNumActors)
-			{
-				CreateNearestNeighborActors(EditorWorld, NearestNeighborActors.Num());
-			}
-
-			for (int32 PartId = 0; PartId < NearestNeighborActors.Num(); PartId++)
-			{
-				const TObjectPtr<UGeometryCache> GeometryCache = GetNearestNeighborModel()->GetNearestNeighborCache(PartId);
-				NearestNeighborActors[PartId]->GetGeometryCacheComponent()->SetGeometryCache(GeometryCache);
 			}
 		}
 	}
 
-	uint8 FNearestNeighborEditorModel::UpdateNearestNeighborData()
+	FVertexMapSelector* FNearestNeighborEditorModel::GetVertexMapSelector() const
+	{
+		return VertexMapSelector.Get();
+	}
+
+	FVertVizSelector* FNearestNeighborEditorModel::GetVertVizSelector() const
 	{
-		UNearestNeighborModel *NearestNeighborModel = static_cast<UNearestNeighborModel*>(Model);
-		check(NearestNeighborModel != nullptr);
+		return VertVizSelector.Get();
+	}
 
-		if (NearestNeighborModel->GetNumParts() == 0)
+	void FNearestNeighborEditorModel::CreateSamplers()
+	{
+		FMLDeformerMorphModelEditorModel::CreateSamplers();
+		const UNearestNeighborModel* const NearestNeighborModel = GetCastModel();
+		if (!NearestNeighborModel)
 		{
-			return EUpdateResult::SUCCESS;
+			return;
 		}
-
-		FMLDeformerSampler* Sampler = GetNumTrainingInputAnims() > 0 ? GetSamplerForTrainingAnim(0) : nullptr;
-		FNearestNeighborGeomCacheSampler* GeomCacheSampler = static_cast<FNearestNeighborGeomCacheSampler*>(Sampler);
-		if (GeomCacheSampler)
+		
+		const int32 NumTrainingAnims = GetNumTrainingInputAnims();
+		const int32 NumSections = NearestNeighborModel->GetNumSections();
+		for (int32 Index = 0; Index < NumSections; ++Index)
 		{
-			TObjectPtr<USkeletalMeshComponent> SkeletalMeshComponent = GeomCacheSampler->GetSkeletalMeshComponent();
-			TObjectPtr<UGeometryCacheComponent> GeometryCacheComponent = GeomCacheSampler->GetGeometryCacheComponent();
-
-			if(SkeletalMeshComponent && GeometryCacheComponent)
-			{
-				UNearestNeighborTrainingModel *TrainingModel = InitTrainingModel<UNearestNeighborTrainingModel>(this);
-				check(TrainingModel != nullptr);
-				uint8 ReturnCode = TrainingModel->UpdateNearestNeighborData();
-				ResetSamplerData();
-				if (HasError(ReturnCode) == 0)
-				{
-					NearestNeighborModel->ValidateNearestNeighborData();
-				}
-				return ReturnCode;
-			}
+			TSharedPtr<FNearestNeighborGeomCacheSampler> Sampler = StaticCastSharedPtr<FNearestNeighborGeomCacheSampler>(CreateSamplerObject());
+			check(Sampler.IsValid());
+			Sampler->Init(this, NumTrainingAnims + Index);
+			Samplers.Add(Sampler);
 		}
-		UE_LOG(LogNearestNeighborModel, Warning, TEXT("GeomCacheSampler is empty. Nearest neighbor data is not updated."))
-		return EUpdateResult::ERROR;
 	}
 
-	template<typename T>
-	T* CreateObjectInstance(const FString& PackageName)
+	bool FNearestNeighborEditorModel::IsAnimIndexValid(int32 Index) const
 	{
-		// Parent package to place new mesh
-		UPackage* Package = nullptr;
-		FString NewPackageName = PackageName;
-		const FString ObjectName = FPaths::GetBaseFilename(PackageName);
-
-		// Setup package name and create one accordingly
-		NewPackageName = UPackageTools::SanitizePackageName(NewPackageName);
-		Package = CreatePackage(*NewPackageName);
-
-		const FString SanitizedObjectName = ObjectTools::SanitizeObjectName(ObjectName);
-
-		T* ExistingTypedObject = FindObject<T>(Package, *SanitizedObjectName);
-		UObject* ExistingObject = FindObject<UObject>(Package, *SanitizedObjectName);
-
-		if (ExistingTypedObject != nullptr)
+		const UNearestNeighborModel* const NearestNeighborModel = GetCastModel();
+		if (!NearestNeighborModel)
 		{
-			ExistingTypedObject->PreEditChange(nullptr);
+			return false;
 		}
-		else if (ExistingObject != nullptr)
-		{
-			// Replacing an object.  Here we go!
-			// Delete the existing object
-			const bool bDeleteSucceeded = ObjectTools::DeleteSingleObject(ExistingObject);
+		const int32 NumSections = NearestNeighborModel->GetNumSections();
+		const int32 NumTrainingAnims = GetNumTrainingInputAnims();
+		return FMath::IsWithin(Index, 0, NumTrainingAnims + NumSections);
+	}
 
-			if (bDeleteSucceeded)
-			{
-				// Force GC so we can cleanly create a new asset (and not do an 'in place' replacement)
-				CollectGarbage(GARBAGE_COLLECTION_KEEPFLAGS);
+	UNearestNeighborModel* FNearestNeighborEditorModel::GetCastModel() const
+	{
+		return Cast<UNearestNeighborModel>(Model.Get());
+	}
 
-				// Create a package for each mesh
-				Package = CreatePackage(*NewPackageName);
-			}
-			else
-			{
-				// failed to delete
-				UE_LOG(LogNearestNeighborModel, Error, TEXT("Failed to delete existing object %s"), *SanitizedObjectName);
-				return nullptr;
-			}
-		}
+	TUniquePtr<FNearestNeighborEditorModelActor> FNearestNeighborEditorModel::CreateNearestNeighborActor(UWorld* World) const
+	{
+		static FLinearColor LabelColor = FNearestNeighborModelEditorStyle::Get().GetColor("NearestNeighborModel.NearestNeighborActors.LabelColor");
+		static FLinearColor WireframeColor = FNearestNeighborModelEditorStyle::Get().GetColor("NearestNeighborModel.NearestNeighborActors.WireframeColor");
+		static FName ActorName = FName("NearestNeighborActor");
+		static FText LabelText = LOCTEXT("NearestNeighborLabelText", "Nearest Neighbor");
+		static int32 ActorID = ActorID_NearestNeighborActors;
 
-		if (Package == nullptr)
-		{
-			UE_LOG(LogNearestNeighborModel, Error, TEXT("Failed to create package %s"), *NewPackageName);
-			return nullptr;
-		}
+		FActorSpawnParameters SpawnParams;
+		SpawnParams.Name = MakeUniqueObjectName(World, AActor::StaticClass(), ActorName);
+		AActor* Actor = World->SpawnActor<AActor>(SpawnParams);
+		Actor->SetFlags(RF_Transient);
+		
+		// Create the Geometry Cache Component.
+		UGeometryCacheComponent* GeomCacheComponent = NewObject<UGeometryCacheComponent>(Actor);
+		GeomCacheComponent->RegisterComponent();
+		GeomCacheComponent->SetOverrideWireframeColor(true);
+		GeomCacheComponent->SetWireframeOverrideColor(WireframeColor);
+		GeomCacheComponent->MarkRenderStateDirty();
+		GeomCacheComponent->SetVisibility(false);
+		Actor->SetRootComponent(GeomCacheComponent);
+		
+		// Create the editor actor.
+		UE::MLDeformer::FMLDeformerEditorActor::FConstructSettings Settings;
+		Settings.Actor = Actor;
+		Settings.TypeID = ActorID;
+		Settings.LabelColor = LabelColor;
+		Settings.LabelText = LabelText;
+		Settings.bIsTrainingActor = false;
 
-		return NewObject<T>(Package, FName(*SanitizedObjectName),  RF_Public | RF_Standalone);
+		TUniquePtr<FNearestNeighborEditorModelActor> NewActor = MakeUnique<FNearestNeighborEditorModelActor>(Settings);
+		NewActor->SetGeometryCacheComponent(GeomCacheComponent);
+		return NewActor;
 	}
 
-	void FNearestNeighborEditorModel::KMeansClusterPoses()
+	void FNearestNeighborEditorModel::UpdateNearestNeighborActor(FNearestNeighborEditorModelActor& Actor) const
 	{
-		KMeansClusterResult = EUpdateResult::SUCCESS;
-		UNearestNeighborModel* NearestNeighborModel = static_cast<UNearestNeighborModel*>(Model);
-		check(NearestNeighborModel != nullptr);
-
-		if (NearestNeighborModel->KMeansPartId >= NearestNeighborModel->GetNumParts())
+		const UNearestNeighborModel* const NearestNeighborModel = GetCastModel();
+		if (!NearestNeighborModel)
 		{
-			UE_LOG(LogNearestNeighborModel, Error, TEXT("KMeansPartId %d is out of range [0, %d). Nothing will be done."), NearestNeighborModel->KMeansPartId, NearestNeighborModel->GetNumParts());
-			KMeansClusterResult |= EUpdateResult::ERROR;
 			return;
 		}
-		if (NearestNeighborModel->SourceAnims.Num() == 0)
+		const UNearestNeighborModelVizSettings* NNViz = GetCastVizSettings();
+		if (!NNViz)
 		{
-			UE_LOG(LogNearestNeighborModel, Error, TEXT("No source anims found."));
-			KMeansClusterResult |= EUpdateResult::ERROR;
 			return;
 		}
-		for (int32 i = 0; i < NearestNeighborModel->SourceAnims.Num(); i++)
-		{
-			if (NearestNeighborModel->SourceAnims[i] == nullptr)
-			{
-				UE_LOG(LogNearestNeighborModel, Error, TEXT("Source anim %d is null."), i);
-				KMeansClusterResult |= EUpdateResult::ERROR;
-				return;
-			}
-		}
-
-		UNearestNeighborTrainingModel *TrainingModel = InitTrainingModel<UNearestNeighborTrainingModel>(this);
-		check(TrainingModel != nullptr);
-		KMeansClusterResult |= TrainingModel->KmeansClusterPoses(NearestNeighborModel->KMeansPartId);
-		if (HasError(KMeansClusterResult))
+		const int32 SectionIndex = NNViz->NearestNeighborActorSectionIndex;
+		const int32 NumSections = NearestNeighborModel->GetNumSections();
+		if (SectionIndex < 0 || SectionIndex >= NumSections)
 		{
 			return;
 		}
-		ResetSamplerData();
-		TArray<int32> KmeansResults = TrainingModel->KmeansResults;
 
-		const FString PackageName = GetTestMLDeformerComponent()->GetDeformerAsset()->GetPackage()->GetName();
-		const FString DirName = FPackageName::GetLongPackagePath(PackageName);
-		const FString FileName = FPaths::GetBaseFilename(PackageName);
-		const FString SavePath = FString::Printf(TEXT("%s/%s_PartId_%d"), *DirName, *FileName, NearestNeighborModel->KMeansPartId);
+		const float Offset = NNViz->NearestNeighborActorsOffset;
+		Actor.SetMeshOffsetFactor(Offset * 2.0f);
 
-		TPair<UAnimSequence*, uint8> AnimAndFlag = CreateAnimOfClusterCenters(SavePath, KmeansResults);
-		UAnimSequence* Anim = AnimAndFlag.Get<0>();
-		KMeansClusterResult |= AnimAndFlag.Get<1>();
-		if (HasError(KMeansClusterResult))
+		UGeometryCache* const GeomCache = NearestNeighborModel->GetSection(SectionIndex).GetMutableNeighborMeshes();
+		if (!GeomCache)
 		{
 			return;
 		}
-
-		UPackage* Package = Anim->GetPackage();
-		const bool bSaveSucced = UPackage::SavePackage(Package, Anim, *SavePath, FSavePackageArgs());
-
-		KmeansResults.Reset();
-	}
-
-	template<typename T>
-	TArray<T> Range(T End)
-	{
-		TArray<T> Result;
-		Result.SetNum(End);
-		for (uint32 i = 0; (T)i < End; i++)
+		Actor.SetGeometryCache(GeomCache);
+		const UMLDeformerComponent* MLDeformerComponent = GetTestMLDeformerComponent();
+		if (!MLDeformerComponent)
 		{
-			Result[i] = i;
+			return;
 		}
-		return Result;
+		Actor.SetTrackedComponent(MLDeformerComponent, SectionIndex);
 	}
 
-	class FAnimEvaluator
+	EOpFlag FNearestNeighborEditorModel::Update()
 	{
-	public:
-		FAnimEvaluator(USkeleton* Skeleton)
+		UNearestNeighborModel* const NearestNeighborModel = GetCastModel();
+		if (!NearestNeighborModel)
 		{
-			const FReferenceSkeleton& ReferenceSkeleton = Skeleton->GetReferenceSkeleton();
-			const int32 NumBones = ReferenceSkeleton.GetNum();
-			TArray<uint16> BoneIndices = Range<uint16>(NumBones);
-			BoneContainer.SetUseRAWData(true);
-			BoneContainer.InitializeTo(BoneIndices, UE::Anim::FCurveFilterSettings(), *Skeleton);
-			OutPose.SetBoneContainer(&BoneContainer);
-			OutCurve.InitFrom(BoneContainer);
+			return EOpFlag::Error;
 		}
-
-		TArray<FTransform> GetBoneTransforms(const UAnimSequence* Anim, int32 Frame)
+		if (!NearestNeighborModel->IsReadyForTraining())
 		{
-			const double Time = FMath::Clamp(Anim->GetSamplingFrameRate().AsSeconds(Frame), 0., (double)Anim->GetPlayLength());
-			FAnimExtractContext ExtractionContext(Time);
-			FAnimationPoseData AnimationPoseData(OutPose, OutCurve, TempAttributes);
-			Anim->GetAnimationPose(AnimationPoseData, ExtractionContext);
-
-			const int32 NumBones = BoneContainer.GetNumBones();
-			TArray<FTransform> Transforms;
-			Transforms.SetNum(NumBones);
-			for (int32 BoneIndex = 0; BoneIndex < NumBones; ++BoneIndex)
+			if (IsTrained())
 			{
-				const FCompactPoseBoneIndex CompactIndex = BoneContainer.MakeCompactPoseIndex(FMeshPoseBoneIndex(BoneIndex));
-				const FTransform BoneTransform = OutPose[CompactIndex];
-				Transforms[BoneIndex] = BoneTransform;
+				NearestNeighborModel->ClearOptimizedNetwork();
 			}
-			return Transforms;
+			return EOpFlag::Error;
 		}
-
-	private:
-		FBoneContainer BoneContainer;
-		FCompactPose OutPose;
-		FBlendedCurve OutCurve;
-		UE::Anim::FStackAttributeContainer TempAttributes;
-	};
-
-	TPair<UAnimSequence*, uint8> FNearestNeighborEditorModel::CreateAnimOfClusterCenters(const FString& PackageName, const TArray<int32>& KmeansResults)
-	{
-		uint8 ReturnCode = EUpdateResult::SUCCESS;
-		UNearestNeighborModel *NearestNeighborModel = static_cast<UNearestNeighborModel*>(Model);
-		TTuple<UAnimSequence*, uint8> None = TTuple<UAnimSequence*, uint8>(nullptr, EUpdateResult::ERROR);
-
-		if (KmeansResults.Num() != NearestNeighborModel->NumClusters * 2)
+		EOpFlag UpdateResult = CheckNetwork();
+		if (OpFlag::HasError(UpdateResult))
 		{
-			UE_LOG(LogNearestNeighborModel, Error, TEXT("KmeansClusterPoses returned %d clusters whereas %d are expected."), KmeansResults.Num() / 2, NearestNeighborModel->NumClusters);
-			return None;
+			return UpdateResult;
 		}
-		if (NearestNeighborModel->SourceAnims.Num() == 0)
+		UpdateResult |= UpdateNearestNeighborData();
+		if (OpFlag::HasError(UpdateResult))
 		{
-			UE_LOG(LogNearestNeighborModel, Error, TEXT("No source anims found."));
-			return None;
+			return UpdateResult;
 		}
-		const UAnimSequence* DefaultAnim = NearestNeighborModel->SourceAnims[0];
-		if (DefaultAnim == nullptr)
+		UpdateResult |= NearestNeighborModel->UpdateForInference();
+		if (OpFlag::HasError(UpdateResult))
 		{
-			UE_LOG(LogNearestNeighborModel, Error, TEXT("Source anim 0 is null."));
-			return None;
+			return UpdateResult;
 		}
-
-		UAnimSequence* Anim = CreateObjectInstance<UAnimSequence>(PackageName);
-		if (Anim == nullptr)
-		{
-			UE_LOG(LogNearestNeighborModel, Error, TEXT("Failed to create AnimSequence."));
-			return None;
-		}
-
-		USkeleton* const Skeleton = DefaultAnim->GetSkeleton();
-		if (Skeleton == nullptr)
-		{
-			UE_LOG(LogNearestNeighborModel, Error, TEXT("Skeleton is null."));
-			return None;
-		}
-		const FReferenceSkeleton& ReferenceSkeleton = Skeleton->GetReferenceSkeleton();
-		const int32 NumBones = ReferenceSkeleton.GetNum();
-
-		Anim->SetSkeleton(Skeleton);
-		IAnimationDataController& Controller = Anim->GetController();
-		Controller.OpenBracket(LOCTEXT("CreateNewAnim_Bracket", "Create New Anim"));
-		Controller.InitializeModel();
-		Anim->ResetAnimation();
-		Anim->SetPreviewMesh(NearestNeighborModel->GetSkeletalMesh());
-		const IAnimationDataModel* AnimData = Anim->GetDataModel();
-		const int32 NumKeys = NearestNeighborModel->NumClusters;
-		Controller.SetNumberOfFrames(NumKeys - 1);
-		Controller.SetFrameRate(FFrameRate(30, 1));
-
-		FMemMark Mark(FMemStack::Get());
-		FAnimEvaluator AnimEval(Skeleton);
-
-		TArray<TArray<FVector3f>> PosKeys;
-		TArray<TArray<FQuat4f>> RotKeys;
-		TArray<TArray<FVector3f>> ScaleKeys;
-		PosKeys.SetNum(NumBones);
-		RotKeys.SetNum(NumBones);
-		ScaleKeys.SetNum(NumBones);
-		for (int32 BoneIndex = 0; BoneIndex < NumBones; ++BoneIndex)
-		{
-			PosKeys[BoneIndex].SetNum(NumKeys);
-			RotKeys[BoneIndex].SetNum(NumKeys);
-			ScaleKeys[BoneIndex].SetNum(NumKeys);
-		}
-
-		for (int32 KeyIndex = 0; KeyIndex < NumKeys; ++KeyIndex)
-		{
-			const int32 PickedAnimId = KmeansResults[KeyIndex * 2];
-			const int32 PickedFrame = KmeansResults[KeyIndex * 2 + 1];
-			if (PickedAnimId < 0 || PickedAnimId >= NearestNeighborModel->SourceAnims.Num())
-			{
-				UE_LOG(LogNearestNeighborModel, Error, TEXT("CreateAnimOfClusterCenters: PickedAnimId %d is out of range."), PickedAnimId);
-				return None;
-			}
-			
-			const UAnimSequence* PickedAnim = NearestNeighborModel->SourceAnims[PickedAnimId];
-			if (PickedAnim == nullptr)
-			{
-				UE_LOG(LogNearestNeighborModel, Error, TEXT("CreateAnimOfClusterCenters: PickedAnim %d is null."), PickedAnimId);
-				return None;
-			}
-
-			const TArray<FTransform> BoneTransforms = AnimEval.GetBoneTransforms(PickedAnim, PickedFrame);
-			for (int32 BoneIndex = 0; BoneIndex < NumBones; ++BoneIndex)
-			{
-				const FTransform& BoneTransform = BoneTransforms[BoneIndex];
-				PosKeys[BoneIndex][KeyIndex] = FVector3f(BoneTransform.GetLocation());
-				RotKeys[BoneIndex][KeyIndex] = FQuat4f(BoneTransform.GetRotation());
-				ScaleKeys[BoneIndex][KeyIndex] = FVector3f(BoneTransform.GetScale3D());
-			}
-		}
-		for (int32 BoneIndex = 0; BoneIndex < NumBones; ++BoneIndex)
+		UpdateNearestNeighborActor(*NearestNeighborActor);
+		UpdateResult |= UpdateMorphDeltas();
+		if (OpFlag::HasError(UpdateResult))
 		{
-			const FName BoneName = ReferenceSkeleton.GetBoneName(BoneIndex);
-			Controller.AddBoneCurve(BoneName);
-			Controller.SetBoneTrackKeys(BoneName, PosKeys[BoneIndex], RotKeys[BoneIndex], ScaleKeys[BoneIndex]);
+			return UpdateResult;
 		}
-
-		Controller.NotifyPopulated();
-		Controller.CloseBracket();
-		return MakeTuple(Anim, ReturnCode);
+		InitEngineMorphTargets(NearestNeighborModel->GetMorphTargetDeltas());
+		NearestNeighborModel->UpdateMorphTargetsLastWriteTime();
+		SetDefaultDeformerGraphIfNeeded();
+		return UpdateResult;
 	}
 
-	namespace Private
+	void FNearestNeighborEditorModel::OnUpdateClicked()
 	{
-		TArray<int32> GetIncludedFrame(TConstArrayView<int32> ExcludedFrames, int32 NumFrames)
+		EOpFlag Result = Update();
+		if (OpFlag::HasError(Result))
 		{
-			TArray<int32> IncludedFrames;
-			TArray<int32> FrameMap;
-			IncludedFrames.Reserve(NumFrames);
-			for (int32 Frame = 0; Frame < NumFrames; Frame++)
-			{
-				if (!ExcludedFrames.Contains(Frame))
-				{
-					IncludedFrames.Add(Frame);
-				}
-			}
-			return IncludedFrames;
+			FMessageDialog::Open(EAppMsgType::Ok, LOCTEXT("UpdateErrorMessage", "Update failed. Please check Output Log for details") , LOCTEXT("UpdateErrorWindowTitle", "Update Results"));
 		}
-
-		TArray<float> CompactCoeffs(TConstArrayView<float> SourceCoeffs, int32 NumCoeffs, const TArray<int32>& IncludedFrames)
+		else if (OpFlag::HasWarning(Result))
 		{
-			check(NumCoeffs > 0);
-			const int32 NumSource = SourceCoeffs.Num() / NumCoeffs;
-			check(SourceCoeffs.Num() == NumSource * NumCoeffs);
-			const int32 NumResult = IncludedFrames.Num();
-			TArray<float> ResultCoeffs;
-			ResultCoeffs.SetNum(NumResult * NumCoeffs);
-			for (int32 FrameIndex = 0; FrameIndex < NumResult; ++FrameIndex)
-			{
-				const int32 Frame = IncludedFrames[FrameIndex];
-				for (int32 CoeffIndex = 0; CoeffIndex < NumCoeffs; ++CoeffIndex)
-				{
-					ResultCoeffs[FrameIndex * NumCoeffs + CoeffIndex] = SourceCoeffs[Frame * NumCoeffs + CoeffIndex];
-				}
-			}
-			return ResultCoeffs;
+			FMessageDialog::Open(EAppMsgType::Ok, LOCTEXT("UpdateWarningMessage", "Update succeeded with warnings. Please check Output Log for details"), LOCTEXT("UpdateWarningWindowTitle", "Update Results"));
 		}
-	};
+	}
 
-	void FNearestNeighborEditorModel::AddFloatArrayToDeltaArray(const TArray<float>& FloatArr, const TArray<uint32>& VertexMap, TArray<FVector3f>& DeltaArr, int32 DeltaArrayOffset, TOptional<TArray<int32>> OptionalIncludedFrames)
+	EOpFlag FNearestNeighborEditorModel::CheckNetwork()
 	{
-		const int32 NumBaseMeshVerts = Model->GetNumBaseMeshVerts();
-		const UNearestNeighborModel* NearestNeighborModel = GetNearestNeighborModel();
-		const int32 PartNumVerts = VertexMap.Num();
-		if (PartNumVerts == 0)
+		const UNearestNeighborModel* const NearestNeighborModel = GetCastModel();
+		if (!NearestNeighborModel)
 		{
-			return;
+			return EOpFlag::Error;
 		}
-		const int32 NumShapes = FloatArr.Num() / (PartNumVerts * 3);
-		if (DeltaArrayOffset == INDEX_NONE)
+
+		if (!NearestNeighborModel->GetOptimizedNetwork().IsValid())
 		{
-			DeltaArrayOffset = DeltaArr.Num();
+			UE_LOG(LogNearestNeighborModel, Error, TEXT("Network is empty. Please (re-)train model."));
+			return EOpFlag::Error;
 		}
-		const TArray<int32> IncludedFrames = OptionalIncludedFrames.IsSet() ? OptionalIncludedFrames.GetValue() : Range<int32>(NumShapes);
-		const int32 NumIncludedFrames = IncludedFrames.Num();
-		DeltaArr.SetNumZeroed(FMath::Max(DeltaArrayOffset + NumIncludedFrames * NumBaseMeshVerts, DeltaArr.Num()), false);
-
-		for(int32 Frame = 0; Frame < IncludedFrames.Num(); ++Frame)
+	
+		const int32 NumNetworkWeights = NearestNeighborModel->GetNumNetworkOutputs();
+		const int32 NumPCACoeffs = NearestNeighborModel->GetTotalNumPCACoeffs();
+		if (NumNetworkWeights != NumPCACoeffs)
 		{
-			const int32 IncludedFrame = IncludedFrames[Frame];
-			for(int32 VertexId = 0; VertexId < PartNumVerts; VertexId++)
-			{
-				const int32 DeltaId = Frame * NumBaseMeshVerts + VertexMap[VertexId];
-				const int32 FloatId = (IncludedFrame * PartNumVerts + VertexId) * 3;
-				DeltaArr[DeltaArrayOffset + DeltaId] = FVector3f(FloatArr[FloatId], FloatArr[FloatId + 1], FloatArr[FloatId + 2]);
-			}
+			UE_LOG(LogNearestNeighborModel, Error, TEXT("Network output dimension %d is not equal to number of morph targets %d. Network needs to be re-trained."), NumNetworkWeights, NumPCACoeffs);
+			return EOpFlag::Error;
 		}
+		return EOpFlag::Success;
 	}
-	
-	void FNearestNeighborEditorModel::OnMorphTargetUpdate()
-	{
-		MorphTargetUpdateResult = EUpdateResult::SUCCESS;
 
-		UNearestNeighborModel* NearestNeighborModel = GetNearestNeighborModel();
-
-		if (!NearestNeighborModel->IsClothPartDataValid())
+	EOpFlag FNearestNeighborEditorModel::UpdateNearestNeighborData()
+	{
+		if (!GetCastModel())
 		{
-			MorphTargetUpdateResult |= NearestNeighborModel->UpdateClothPartData();
-			if (HasError(MorphTargetUpdateResult))
-			{
-				UpdateNearestNeighborActors();
-				return;
-			}
+			return EOpFlag::Error;
 		}
-		MorphTargetUpdateResult |= WarnIfNetworkInvalid();
-
-		if (!NearestNeighborModel->IsNearestNeighborDataValid())
+		if (!IsTrained())
 		{
-			MorphTargetUpdateResult |= UpdateNearestNeighborData();
-			UpdateNearestNeighborActors();
-			if (HasError(MorphTargetUpdateResult))
-			{
-				return;
-			}
+			UE_LOG(LogNearestNeighborModel, Error, TEXT("Network is not trained. Nearest neighbor data cannot be updated."));
+			return EOpFlag::Error;
 		}
-
-		if (IsTrained())
+		UNearestNeighborTrainingModel *TrainingModel = FHelpers::GetDerivedCDO<UNearestNeighborTrainingModel>();
+		if (!TrainingModel)
 		{
-			MorphTargetUpdateResult |= InitMorphTargets();
-			if (HasError(MorphTargetUpdateResult))
-			{
-				return;
-			}
-			RefreshMorphTargets();
-			GetNearestNeighborModel()->UpdateNetworkSize();
-			GetNearestNeighborModel()->UpdateMorphTargetSize();
+			return EOpFlag::Error;
 		}
+		TrainingModel->Init(this);
 
-		InitTestMLDeformerPreviousWeights();
-		GetNearestNeighborModel()->ValidateMorphTargetData();
-
-		GetNearestNeighborModel()->UpdateInputMultipliers();
+		return ToOpFlag(TrainingModel->UpdateNearestNeighborData());
 	}
 
-	uint8 FNearestNeighborEditorModel::InitMorphTargets()
+	EOpFlag FNearestNeighborEditorModel::UpdateMorphDeltas()
 	{
-		uint8 Result = EUpdateResult::SUCCESS;
-		UNearestNeighborModel* NearestNeighborModel = GetNearestNeighborModel();
-
+		const UNearestNeighborModel* const NearestNeighborModel = GetCastModel();
+		if (!NearestNeighborModel)
+		{
+			return EOpFlag::Error;
+		}
+		EOpFlag Result = EOpFlag::Success;
+	
 		const USkeletalMesh* SkelMesh = Model->GetSkeletalMesh();
 		const int32 NumBaseMeshVerts = Model->GetNumBaseMeshVerts();
 		if (!SkelMesh || NumBaseMeshVerts == 0)
 		{
 			UE_LOG(LogNearestNeighborModel, Error, TEXT("SkeletalMesh is empty. No morph targets are generated"));
-			return EUpdateResult::ERROR;
+			return EOpFlag::Error;
 		}
-
+	
 		if (Model->GetVertexMap().IsEmpty())
 		{
 			UE_LOG(LogNearestNeighborModel, Error, TEXT("VertexMap of the skeletal mesh is empty. No morph targets are generated"));
-			return EUpdateResult::ERROR;
+			return EOpFlag::Error;
 		}
-
-		const int32 NumParts = NearestNeighborModel->GetNumParts();
-		if (NumParts == 0)
+	
+		const int32 NumSections = NearestNeighborModel->GetNumSections();
+		if (NumSections == 0)
 		{
-			UE_LOG(LogNearestNeighborModel, Error, TEXT("There are no cloth parts. No morph targets are generated"));
-			return EUpdateResult::ERROR;
+			UE_LOG(LogNearestNeighborModel, Error, TEXT("There are no cloth sections. No morph targets are generated"));
+			return EOpFlag::Error;
 		}
-
+	
 		const int32 NumImportedModelVerts = FMath::Max(NearestNeighborModel->GetVertexMap()) + 1;
 		if (NumImportedModelVerts != NumBaseMeshVerts)
 		{
 			UE_LOG(LogNearestNeighborModel, Error, TEXT("Vertex count mismatch: imported model of SkeletalMesh has %d vertices and cached SkeletalMesh has %d vertices"), NumImportedModelVerts, NumBaseMeshVerts);
-			return EUpdateResult::ERROR;
+			return EOpFlag::Error;
 		}
-
+	
 		TArray<FVector3f> Deltas;
-		Deltas.Reset();
-
-		int32 NumPCACoeff = 0;
-		for (int32 PartId = 0; PartId < NumParts; PartId++)
-		{
-			NumPCACoeff += NearestNeighborModel->GetPCACoeffNum(PartId);
-		}
-		Deltas.Reserve((1 + NumPCACoeff) * NumBaseMeshVerts);
-
-		for (int32 PartId = 0; PartId < NumParts; PartId++)
+	
+		const int32 TotalNumPCACoeffs = NearestNeighborModel->GetTotalNumPCACoeffs();
+		const int32 TotalNumNeighbors = NearestNeighborModel->GetTotalNumNeighbors();
+		Deltas.SetNumZeroed((1 + TotalNumPCACoeffs + TotalNumNeighbors) * NumBaseMeshVerts);
+	
+		using FSection = UNearestNeighborModel::FSection;
+		using Private::AddFloatArrayToDeltaArray;
+		int32 MorphOffset = 1;
+		for (int32 SectionIndex = 0; SectionIndex < NearestNeighborModel->GetNumSections(); ++SectionIndex)
 		{
-			const TArray<uint32>& VertexMap = NearestNeighborModel->PartVertexMap(PartId);
+			const FSection& Section = NearestNeighborModel->GetSection(SectionIndex);
+			TConstArrayView<int32> VertexMap = Section.GetVertexMap();
 			if (VertexMap.IsEmpty())
 			{
-				UE_LOG(LogNearestNeighborModel, Warning, TEXT("Cloth part %d has empty vertex map. No morph targets are generated for this part."), PartId);
-				Result = EUpdateResult::WARNING;
+				UE_LOG(LogNearestNeighborModel, Warning, TEXT("Section %d has empty vertex map. No morph targets are generated for this section"), SectionIndex);
+				Result |= EOpFlag::Warning;
+				continue;
+			}
+			TArrayView<FVector3f> MeanDeltas(Deltas.GetData(), NumBaseMeshVerts);
+			AddFloatArrayToDeltaArray(Section.GetVertexMean(), VertexMap, MeanDeltas);
+			const int32 NumPCACoeffs = Section.GetNumPCACoeffs();
+			const int32 SectionNumVerts = Section.GetNumVertices();
+			check(VertexMap.Num() == SectionNumVerts);
+			for (int32 Index = 0; Index < NumPCACoeffs; ++Index)
+			{
+				TConstArrayView<float> BasisFloats(Section.GetPCABasis().GetData() + Index * SectionNumVerts * 3, SectionNumVerts * 3);
+				TArrayView<FVector3f> BasisDeltas(Deltas.GetData() + (MorphOffset + Index) * NumBaseMeshVerts, NumBaseMeshVerts);
+				AddFloatArrayToDeltaArray(BasisFloats, VertexMap, BasisDeltas);
 			}
-			AddFloatArrayToDeltaArray(NearestNeighborModel->ClothPartData[PartId].VertexMean, VertexMap, Deltas, 0);
-			AddFloatArrayToDeltaArray(NearestNeighborModel->ClothPartData[PartId].PCABasis, VertexMap, Deltas);
+			MorphOffset += NumPCACoeffs;
 		}
-
-		for (int32 PartId = 0; PartId < NumParts; PartId++)
+	
+		for (int32 SectionIndex = 0; SectionIndex < NearestNeighborModel->GetNumSections(); ++SectionIndex)
 		{
-			FClothPartData& Data = NearestNeighborModel->ClothPartData[PartId];
-			const TArray<uint32>& VertexMap = Data.VertexMap;
-			const TArray<float>& AssetNeighborOffsets = Data.AssetNeighborOffsets;
-			constexpr int32 DeltaOffset = INDEX_NONE;
-			const int32 NumAssetNeighbors = AssetNeighborOffsets.Num() / (VertexMap.Num() * 3);
-			check(AssetNeighborOffsets.Num() % (VertexMap.Num() * 3) == 0);
-			if (NumAssetNeighbors == 0)
+			const FSection& Section = NearestNeighborModel->GetSection(SectionIndex);
+			TConstArrayView<int32> VertexMap = Section.GetVertexMap();
+			if (VertexMap.IsEmpty())
+			{
+				// Warning already generated
+				continue;
+			}
+			if (!Section.IsReadyForInference())
 			{
+				UE_LOG(LogNearestNeighborModel, Warning, TEXT("Section %d is not ready for inference. No morph targets are generated for this section"), SectionIndex);
+				Result |= EOpFlag::Warning;
 				continue;
 			}
-			check(NearestNeighborModel->NearestNeighborData.IsValidIndex(PartId));
-			const TArray<int32> ExcludedFrames = NearestNeighborModel->NearestNeighborData[PartId].ExcludedFrames;
-			const TArray<int32> IncludedFrames = UE::NearestNeighborModel::Private::GetIncludedFrame(ExcludedFrames, NumAssetNeighbors);
-			NearestNeighborModel->SetNumNeighbors(PartId, IncludedFrames.Num());
-			const int32 NumPCACoeffs = Data.PCACoeffNum;
-			Data.NeighborCoeffs = Private::CompactCoeffs(Data.AssetNeighborCoeffs, NumPCACoeffs, IncludedFrames);
-			Data.AssetNeighborIndexMap = IncludedFrames;
-			AddFloatArrayToDeltaArray(AssetNeighborOffsets, VertexMap, Deltas, DeltaOffset, IncludedFrames);
+			TArray<float> RuntimeNeighborOffsets;
+			Section.GetRuntimeNeighborOffsets(RuntimeNeighborOffsets);
+			const int32 SectionNumVerts = Section.GetNumVertices();
+			const int32 SectionNumNeighbors = Section.GetRuntimeNumNeighbors();
+			check(RuntimeNeighborOffsets.Num() == SectionNumVerts * 3 * SectionNumNeighbors);
+			for (int32 Index = 0; Index < SectionNumNeighbors; ++Index)
+			{
+				TConstArrayView<float> NeighborOffsets(RuntimeNeighborOffsets.GetData() + Index * SectionNumVerts * 3, SectionNumVerts * 3);
+				TArrayView<FVector3f> NeighborDeltas(Deltas.GetData() + (MorphOffset + Index) * NumBaseMeshVerts, NumBaseMeshVerts);
+				AddFloatArrayToDeltaArray(NeighborOffsets, VertexMap, NeighborDeltas);
+			}
 		}
-
+	
 		if (Deltas.Num() == 0)
 		{
 			UE_LOG(LogNearestNeighborModel, Error, TEXT("All cloth parts are empty. No morph targets are generated."));
-			return EUpdateResult::ERROR;
+			return EOpFlag::Error;
 		}
 
-		const int32 LOD = 0;
-		TArray<UMorphTarget*> MorphTargets;
-		CreateMorphTargets(
-			MorphTargets, 
-			Deltas, 
-			FString("NNMorphTarget_"),
-			LOD,
-			NearestNeighborModel->GetMorphDeltaZeroThreshold(),
-			NearestNeighborModel->GetIncludeMorphTargetNormals(),
-			NearestNeighborModel->GetMaskChannel(),
-			NearestNeighborModel->GetInvertMaskChannel());
-
-		check(NearestNeighborModel->GetMorphTargetSet().IsValid());
-		FMorphTargetVertexInfoBuffers& MorphBuffers = NearestNeighborModel->GetMorphTargetSet()->MorphBuffers;
-		CompressMorphTargets(MorphBuffers, MorphTargets, LOD, NearestNeighborModel->GetMorphCompressionLevel());
-
-		if (MorphBuffers.GetNumBatches() <= 0)
-		{
-			UE_LOG(LogNearestNeighborModel, Warning, TEXT("Morph buffer is empty. It is possible that all deltas are zero. No morph targets are generated."));
-			Result = EUpdateResult::WARNING;
-			NearestNeighborModel->ResetMorphBuffers();
-		}
+		GetCastModel()->SetMorphTargetDeltas(Deltas);
+		return Result;
+	}
 
-		// Remove the morph targets again, as we don't need them anymore.
-		for (UMorphTarget* MorphTarget : MorphTargets)
+	void FNearestNeighborEditorModel::ResetMorphTargets()
+	{
+		const TSharedPtr<FExternalMorphSet> MorphSet = GetMorphModel()->GetMorphTargetSet();
+		if (MorphSet.IsValid())
 		{
-			MorphTarget->ConditionalBeginDestroy();
+			MorphSet->MorphBuffers = FMorphTargetVertexInfoBuffers();
 		}
-
-		NearestNeighborModel->UpdateStatistics();
-
-		NearestNeighborModel->SetMorphTargetDeltas(Deltas);
-		return Result;
 	}
 
-	void FNearestNeighborEditorModel::RefreshMorphTargets()
+	void FNearestNeighborEditorModel::UpdateNearestNeighborIds()
 	{
-		USkeletalMeshComponent* SkelMeshComponent = FindEditorActor(ActorID_Test_MLDeformed)->GetSkeletalMeshComponent() ;
-		if (SkelMeshComponent)
+		const UNearestNeighborModelInstance* ModelInstance = GetTestNearestNeighborModelInstance();
+		const UNearestNeighborModel* NearestNeighborModel = GetCastModel();
+		UNearestNeighborModelVizSettings* NNViz = GetCastVizSettings();
+		if (ModelInstance && NearestNeighborModel && NNViz)
 		{
-			check(GetNearestNeighborModel()->GetMorphTargetSet().IsValid());
-			FMorphTargetVertexInfoBuffers& MorphBuffers = GetNearestNeighborModel()->GetMorphTargetSet()->MorphBuffers;
-			BeginReleaseResource(&MorphBuffers);
-			if (MorphBuffers.IsMorphCPUDataValid() && MorphBuffers.GetNumMorphs() > 0 && MorphBuffers.GetNumBatches() > 0)
+			const TArray<uint32>& NeighborIds = ModelInstance->GetNearestNeighborIds();
+			const int32 Num = FMath::Min(NeighborIds.Num(), NearestNeighborModel->GetNumSections());
+			TArray<int32> AssetNeighborIds;
+			AssetNeighborIds.SetNum(Num);
+			for(int32 Index = 0; Index < Num; ++Index)
 			{
-				BeginInitResource(&MorphBuffers);
+				const int32 NeighborId = NeighborIds[Index];
+				const TArray<int32>& IndexMap = NearestNeighborModel->GetSection(Index).GetAssetNeighborIndexMap();
+				AssetNeighborIds[Index] = IndexMap.IsValidIndex(NeighborId) ? IndexMap[NeighborId] : INDEX_NONE;
 			}
-			SkelMeshComponent->RefreshExternalMorphTargetWeights();
+			NNViz->NearestNeighborIds = AssetNeighborIds;
 		}
 	}
 
+	UNearestNeighborModelVizSettings* FNearestNeighborEditorModel::GetCastVizSettings() const
+	{
+		return Cast<UNearestNeighborModelVizSettings>(GetMorphModel()->GetVizSettings()); 
+	}
+
 	UMLDeformerComponent* FNearestNeighborEditorModel::GetTestMLDeformerComponent() const
 	{
-		return FindMLDeformerComponent(ActorID_Test_MLDeformed);
+		return FindMLDeformerComponent(UE::MLDeformer::ActorID_Test_MLDeformed);
 	}
 
+	USkeletalMeshComponent* FNearestNeighborEditorModel::GetTestSkeletalMeshComponent() const
+	{
+		if (const UMLDeformerComponent* MLDeformerComponent = GetTestMLDeformerComponent())
+		{
+			return MLDeformerComponent->GetSkeletalMeshComponent();
+		}
+		return nullptr;
+	}
+	
+	UNearestNeighborModelInstance* FNearestNeighborEditorModel::GetTestNearestNeighborModelInstance() const
+	{
+		if (const UMLDeformerComponent* MLDeformerComponent = FindMLDeformerComponent(UE::MLDeformer::ActorID_Test_MLDeformed))
+		{
+			return Cast<UNearestNeighborModelInstance>(MLDeformerComponent->GetModelInstance());
+		}
+		return nullptr;
+	}
 
-	void FNearestNeighborEditorModel::InitTestMLDeformerPreviousWeights() 
+	void FVertexMapSelector::Update(const USkeletalMesh* SkelMesh)
 	{
-		UNearestNeighborModelInstance* ModelInstance = static_cast<UNearestNeighborModelInstance*>(GetTestMLDeformerModelInstance());
-		if (ModelInstance)
+		if (!SkelMesh)
+		{
+			Reset();
+			return;
+		}
+		const FSkeletalMeshModel* ImportedModel = SkelMesh->GetImportedModel();
+		if (!ImportedModel)
+		{
+			Reset();
+			return;
+		}
+		constexpr int32 LODIndex = 0;
+		if (!ImportedModel->LODModels.IsValidIndex(LODIndex))
+		{
+			Reset();
+			return;
+		}
+		const FSkeletalMeshLODModel& LODModel = ImportedModel->LODModels[LODIndex];
+		const TArray<FSkelMeshImportedMeshInfo>& SkelMeshInfos = LODModel.ImportedMeshInfos;
+		Options.Reserve(SkelMeshInfos.Num() + 1);
+		VertexMapStrings.Reserve(SkelMeshInfos.Num() + 1);
+		for (const FSkelMeshImportedMeshInfo& Info: SkelMeshInfos)
 		{
-			ModelInstance->InitPreviousWeights();
+			TSharedPtr<FString> Option = MakeShared<FString>(Info.Name.ToString());
+			Options.Add(Option);
+			FString VertexMapString = FString::Printf(TEXT("%d-%d"), Info.StartImportedVertex, Info.StartImportedVertex + Info.NumVertices - 1);
+			VertexMapStrings.Add(Option, MoveTemp(VertexMapString));
 		}
+		Options.Add(CustomString);
+		VertexMapStrings.Add(CustomString, FString::Printf(TEXT("0-%d"), LODModel.MaxImportVertex));
 	}
 
-	uint8 FNearestNeighborEditorModel::WarnIfNetworkInvalid()
+	TArray<TSharedPtr<FString>>* FVertexMapSelector::GetOptions()
 	{
-		const UNearestNeighborModel *NearestNeighborModel = static_cast<UNearestNeighborModel*>(Model);
+		return &Options;
+	}
 
-		if (NearestNeighborModel->DoesUseOptimizedNetwork())
+	void FVertexMapSelector::OnSelectionChanged(UNearestNeighborModelSection& Section, TSharedPtr<FString> InSelectedItem, ESelectInfo::Type SelectInfo) const
+	{
+		TSharedPtr<FString> SelectedItem = GetSelectedItem(Section);
+		if (InSelectedItem != SelectedItem)
 		{
-			if (NearestNeighborModel->GetOptimizedNetwork() == nullptr)
+			Section.SetVertexMapString(VertexMapStrings[InSelectedItem]);
+			const int32 Index = Options.Find(InSelectedItem);
+			if (Index >= 0 && Index < Options.Num() - 1)
 			{
-				UE_LOG(LogNearestNeighborModel, Warning, TEXT("Optimized network is not set. Model needs to be re-trained."));
-				return EUpdateResult::WARNING;
+				Section.SetMeshIndex(Index);
 			}
-
-			const int32 NumNetworkWeights = NearestNeighborModel->GetOptimizedNetworkNumOutputs();
-			const int32 NumPCACoeffs = NearestNeighborModel->GetTotalNumPCACoeffs();
-			if (NumNetworkWeights != NumPCACoeffs)
+			else
 			{
-				UE_LOG(LogNearestNeighborModel, Warning, TEXT("Network output dimension %d is not equal to number of morph targets %d. Network needs to be re-trained and no deformation will be applied."), NumNetworkWeights, NumPCACoeffs);
-				return EUpdateResult::WARNING;
+				Section.SetMeshIndex(INDEX_NONE);
 			}
 		}
-		return EUpdateResult::SUCCESS;
 	}
 
+	TSharedPtr<FString> FVertexMapSelector::GetSelectedItem(const UNearestNeighborModelSection& Section) const
+	{
+		const int32 MeshIndex = Section.GetMeshIndex();
+		if (Options.IsValidIndex(MeshIndex))
+		{
+			return Options[MeshIndex];
+		}
+		else if (!Options.IsEmpty())
+		{
+			return Options.Last();
+		}
+		else
+		{
+			return nullptr;
+		}
+	}
 
-	UMLDeformerModelInstance* FNearestNeighborEditorModel::GetTestMLDeformerModelInstance() const
+	FString FVertexMapSelector::GetVertexMapString(const UNearestNeighborModelSection& Section) const
 	{
-		UMLDeformerModelInstance* ModelInstance = nullptr;
-		const UMLDeformerComponent* MLDeformerComponent = FindMLDeformerComponent(ActorID_Test_MLDeformed);
-		if (MLDeformerComponent)
+		const int32 MeshIndex = Section.GetMeshIndex();
+		if (Options.IsValidIndex(MeshIndex))
+		{
+			return VertexMapStrings[Options[MeshIndex]];
+		}
+		else if (!Options.IsEmpty())
 		{
-			ModelInstance = MLDeformerComponent->GetModelInstance();
+			return VertexMapStrings[Options.Last()];
+		}
+		else
+		{
+			return FString();
 		}
-		return ModelInstance;
 	}
 
-	void FNearestNeighborEditorModel::GetNeighborStats()
+	bool FVertexMapSelector::IsValid() const
 	{
-		UNearestNeighborTrainingModel *TrainingModel = InitTrainingModel<UNearestNeighborTrainingModel>(this);
-		if (!TrainingModel)
+		return Options.Num() > 0 && Options.Num() == VertexMapStrings.Num();
+	}
+
+	void FVertexMapSelector::Reset()
+	{
+		Options.Reset();
+		VertexMapStrings.Reset();
+	}
+
+	TSharedPtr<FString> FVertexMapSelector::CustomString = MakeShared<FString>(TEXT("Custom"));
+
+	FVertVizSelector::FVertVizSelector(UNearestNeighborModelVizSettings* InSettings)
+		: Settings(InSettings)
+	{
+		Options.Add(MakeShared<FString>(TEXT("All Sections")));
+	}
+
+	void FVertVizSelector::Update(int32 NumSections)
+	{
+		check (NumSections >= 0);
+		if (NumSections < Options.Num() - 1)
 		{
-			return;
+			Options.SetNum(NumSections + 1);
 		}
-		const UNearestNeighborModel *NearestNeighborModel = GetNearestNeighborModel();
-		if (!NearestNeighborModel)
+		else
+		{
+			for (int32 i = Options.Num() - 1; i < NumSections; ++i)
+			{
+				Options.Add(MakeShared<FString>(FString::Printf(TEXT("Section %d"), i)));
+			}
+		}
+	}
+	
+	TArray<TSharedPtr<FString>>* FVertVizSelector::GetOptions()
+	{
+		return &Options;
+	}
+	
+	void FVertVizSelector::OnSelectionChanged(TSharedPtr<FString> InSelectedItem, ESelectInfo::Type SelectInfo) const
+	{
+		if (!Settings)
 		{
 			return;
 		}
-		if (!NearestNeighborModel->GetOptimizedNetwork())
+		const int32 Index = Options.Find(InSelectedItem);
+		if (Index != INDEX_NONE)
+		{
+			Settings->VertVizSectionIndex = Index - 1;
+		}
+	}
+	
+	TSharedPtr<FString> FVertVizSelector::GetSelectedItem() const
+	{
+		if (!Settings)
+		{
+			return nullptr;
+		}
+		return Options[Settings->VertVizSectionIndex + 1];
+	}
+
+	int32 FVertVizSelector::GetSectionIndex(TSharedPtr<FString> Item) const
+	{
+		return Options.Find(Item) - 1;
+	}
+
+	void FVertVizSelector::SelectSection(int32 SectionIndex)
+	{
+		if (!Settings)
 		{
-			UE_LOG(LogNearestNeighborModel, Error, TEXT("Network is not loaded. Network is needed to compute stats. Please train model first."));
 			return;
 		}
-		TrainingModel->GetNeighborStats(0);
+		if (SectionIndex < INDEX_NONE || SectionIndex >= Options.Num() - 1)
+		{
+			Settings->VertVizSectionIndex = INDEX_NONE; // set to default
+		}
+		else
+		{
+			Settings->VertVizSectionIndex = SectionIndex;
+		}
 	}
 }	// namespace UE::NearestNeighborModel
 
diff --git a/Engine/Plugins/Animation/MLDeformer/NearestNeighborModel/Source/NearestNeighborModelEditor/Private/NearestNeighborEditorModelActor.cpp b/Engine/Plugins/Animation/MLDeformer/NearestNeighborModel/Source/NearestNeighborModelEditor/Private/NearestNeighborEditorModelActor.cpp
--- a/Engine/Plugins/Animation/MLDeformer/NearestNeighborModel/Source/NearestNeighborModelEditor/Private/NearestNeighborEditorModelActor.cpp
+++ b/Engine/Plugins/Animation/MLDeformer/NearestNeighborModel/Source/NearestNeighborModelEditor/Private/NearestNeighborEditorModelActor.cpp
@@ -1,8 +1,10 @@
 // Copyright Epic Games, Inc. All Rights Reserved.
 
 #include "NearestNeighborEditorModelActor.h"
-#include "NearestNeighborModelInstance.h"
+
 #include "GeometryCacheComponent.h"
+#include "MLDeformerComponent.h"
+#include "NearestNeighborModelInstance.h"
 
 namespace UE::NearestNeighborModel
 {
@@ -11,21 +13,34 @@ namespace UE::NearestNeighborModel
 	{
 	}
 
-	void FNearestNeighborEditorModelActor::InitNearestNeighborActor(const int32 InPartId, const UMLDeformerComponent* InComponent)
+	FNearestNeighborEditorModelActor::~FNearestNeighborEditorModelActor() = default;
+
+	void FNearestNeighborEditorModelActor::SetGeometryCache(UGeometryCache* InGeometryCache) const
+	{
+		if (GeomCacheComponent)
+		{
+			GeomCacheComponent->SetGeometryCache(InGeometryCache);
+		}
+	}
+
+	void FNearestNeighborEditorModelActor::SetTrackedComponent(const UMLDeformerComponent* InComponent, int32 InSectionIndex)
 	{
-		PartId = InPartId;
-		MLDeformerComponent = InComponent;
+		TrackedComponent = InComponent;
+		SectionIndex = InSectionIndex;
 	}
 
-	void FNearestNeighborEditorModelActor::TickNearestNeighborActor()
+	void FNearestNeighborEditorModelActor::Tick() const
 	{
-		if (GeomCacheComponent && GeomCacheComponent->GetGeometryCache() && MLDeformerComponent)
+		if (GeomCacheComponent && GeomCacheComponent->GetGeometryCache() && TrackedComponent.IsValid())
 		{
-			const UNearestNeighborModelInstance* ModelInstance = static_cast<UNearestNeighborModelInstance*>(MLDeformerComponent->GetModelInstance());
-			if (ModelInstance && PartId < ModelInstance->NeighborIdNum())
+			if (const UNearestNeighborModelInstance* ModelInstance = static_cast<UNearestNeighborModelInstance*>(TrackedComponent->GetModelInstance()))
 			{
-				GeomCacheComponent->SetManualTick(true);
-				GeomCacheComponent->TickAtThisTime(GeomCacheComponent->GetTimeAtFrame(ModelInstance->NearestNeighborId(PartId)), false, false, false);
+				const TArray<uint32> NearestNeighborIds = ModelInstance->GetNearestNeighborIds();
+				if (NearestNeighborIds.IsValidIndex(SectionIndex))
+				{
+					GeomCacheComponent->SetManualTick(true);
+					GeomCacheComponent->TickAtThisTime(GeomCacheComponent->GetTimeAtFrame(NearestNeighborIds[SectionIndex]), false, false, false);
+				}
 			}
 		}
 	}
diff --git a/Engine/Plugins/Animation/MLDeformer/NearestNeighborModel/Source/NearestNeighborModelEditor/Private/NearestNeighborGeomCacheSampler.cpp b/Engine/Plugins/Animation/MLDeformer/NearestNeighborModel/Source/NearestNeighborModelEditor/Private/NearestNeighborGeomCacheSampler.cpp
--- a/Engine/Plugins/Animation/MLDeformer/NearestNeighborModel/Source/NearestNeighborModelEditor/Private/NearestNeighborGeomCacheSampler.cpp
+++ b/Engine/Plugins/Animation/MLDeformer/NearestNeighborModel/Source/NearestNeighborModelEditor/Private/NearestNeighborGeomCacheSampler.cpp
@@ -1,29 +1,27 @@
 // Copyright Epic Games, Inc. All Rights Reserved.
 
 #include "NearestNeighborGeomCacheSampler.h"
+
+#include "Animation/AnimSequence.h"
+#include "Animation/AnimSingleNodeInstance.h"
+#include "Animation/DebugSkelMeshComponent.h"
 #include "BoneWeights.h"
-#include "NearestNeighborModel.h"
 #include "Engine/SkeletalMesh.h"
-#include "Animation/DebugSkelMeshComponent.h"
-#include "Animation/AnimSequence.h"
-#include "Rendering/SkeletalMeshModel.h"
-#include "Rendering/SkeletalMeshLODModel.h"
-#include "Rendering/SkeletalMeshLODRenderData.h"
-#include "Rendering/SkeletalMeshRenderData.h"
-#include "GeometryCacheComponent.h"
 #include "GeometryCache.h"
+#include "GeometryCacheComponent.h"
 #include "GeometryCacheMeshData.h"
 #include "GeometryCacheTrack.h"
+#include "NearestNeighborModel.h"
+#include "Rendering/SkeletalMeshModel.h"
+#include "Rendering/SkeletalMeshRenderData.h"
 
-using namespace UE::MLDeformer;
 namespace UE::NearestNeighborModel
 {
-	FNearestNeighborGeomCacheSampler::~FNearestNeighborGeomCacheSampler() = default;
-	
 	void FNearestNeighborGeomCacheSampler::Sample(int32 InAnimFrameIndex)
 	{
-		UNearestNeighborModel* NearestNeighborModel = static_cast<UNearestNeighborModel*>(Model);
-		if (NearestNeighborModel->DoesUseDualQuaternionDeltas() && VertexDeltaSpace == EVertexDeltaSpace::PreSkinning)
+		const UNearestNeighborModel* const NearestNeighborModel = static_cast<UNearestNeighborModel*>(Model);
+		using UE::MLDeformer::EVertexDeltaSpace;
+		if (NearestNeighborModel && NearestNeighborModel->DoesUseDualQuaternionDeltas() && VertexDeltaSpace == EVertexDeltaSpace::PreSkinning)
 		{
 			FMLDeformerSampler::Sample(InAnimFrameIndex);
 			UpdateSkinnedPositions();
@@ -35,99 +33,196 @@ namespace UE::NearestNeighborModel
 		}
 	}
 
-	void FNearestNeighborGeomCacheSampler::SampleDualQuaternionDeltas(int32 InAnimFrameIndex)
+	namespace Private
 	{
-		USkeletalMesh* SkeletalMesh = SkeletalMeshComponent.Get() ? SkeletalMeshComponent->GetSkeletalMeshAsset() : nullptr;
-		UGeometryCache* GeometryCache = GeometryCacheComponent->GetGeometryCache();
-		if (SkeletalMeshComponent && SkeletalMesh && GeometryCacheComponent && GeometryCache)
+		UAnimSequence* GetAnimSequence(const USkeletalMeshComponent* SkeletalMeshComponent)
 		{
-			const FTransform& AlignmentTransform = Model->GetAlignmentTransform();
-			FSkeletalMeshModel* ImportedModel = SkeletalMesh->GetImportedModel();
-
-			// For all mesh mappings we found.
-			const int32 LODIndex = 0;
-			const FSkeletalMeshLODModel& LODModel = ImportedModel->LODModels[LODIndex];
-			const TArray<FSkelMeshImportedMeshInfo>& SkelMeshInfos = LODModel.ImportedMeshInfos;
-			for (int32 MeshMappingIndex = 0; MeshMappingIndex < MeshMappings.Num(); ++MeshMappingIndex)
+			if (SkeletalMeshComponent)
 			{
-				const ::UE::MLDeformer::FMLDeformerGeomCacheMeshMapping& MeshMapping = MeshMappings[MeshMappingIndex];
-				const FSkelMeshImportedMeshInfo& MeshInfo = SkelMeshInfos[MeshMapping.MeshIndex];
-				UGeometryCacheTrack* Track = GeometryCache->Tracks[MeshMapping.TrackIndex];
-
-				// Sample the mesh data of the geom cache.
-				FGeometryCacheMeshData& GeomCacheMeshData = GeomCacheMeshDatas[MeshMappingIndex];
-				if (!Track->GetMeshDataAtSampleIndex(InAnimFrameIndex, GeomCacheMeshData))
+				if (UAnimSingleNodeInstance* SingleNodeInstance =  Cast<UAnimSingleNodeInstance>(SkeletalMeshComponent->GetAnimInstance()))
 				{
-					continue;
-				}
-
-				// Calculate the vertex deltas.
-				const FSkeletalMeshLODRenderData& SkelMeshLODData = SkeletalMesh->GetResourceForRendering()->LODRenderData[LODIndex];
-				const FSkinWeightVertexBuffer& SkinWeightBuffer = *SkeletalMeshComponent->GetSkinWeightBuffer(LODIndex);
-				for (int32 VertexIndex = 0; VertexIndex < MeshInfo.NumVertices; ++VertexIndex)
-				{
-					const int32 SkinnedVertexIndex = MeshInfo.StartImportedVertex + VertexIndex;
-					const int32 GeomCacheVertexIndex = MeshMapping.SkelMeshToTrackVertexMap[VertexIndex];
-					if (GeomCacheVertexIndex != INDEX_NONE && GeomCacheMeshData.Positions.IsValidIndex(GeomCacheVertexIndex))
-					{
-						FVector3f Delta = FVector3f::ZeroVector;
-						const int32 ArrayIndex = 3 * SkinnedVertexIndex;
-						const int32 RenderVertexIndex = MeshMapping.ImportedVertexToRenderVertexMap[VertexIndex];
-						if (RenderVertexIndex != INDEX_NONE)
-						{
-							const FVector3f SkinnedVertexPos = SkinnedVertexPositions[SkinnedVertexIndex];
-							const FVector3f GeomCacheVertexPos = (FVector3f)AlignmentTransform.TransformPosition((FVector)GeomCacheMeshData.Positions[GeomCacheVertexIndex]);
-							const FVector3f WorldDelta = GeomCacheVertexPos - SkinnedVertexPos;
-							Delta = CalcDualQuaternionDelta(RenderVertexIndex, WorldDelta, SkelMeshLODData, SkinWeightBuffer);
-						}
-
-						VertexDeltas[ArrayIndex] = Delta.X;
-						VertexDeltas[ArrayIndex + 1] = Delta.Y;
-						VertexDeltas[ArrayIndex + 2] = Delta.Z;
-					}
+					return Cast<UAnimSequence>(SingleNodeInstance->GetCurrentAsset());
 				}
 			}
-		}
-		else
-		{
-			VertexDeltas.Reset(0);
+			return nullptr;
 		}
 	}
 
-	template<typename QuatType>
-	QuatType Conjugate(const QuatType& Q)
+	bool FNearestNeighborGeomCacheSampler::CustomSample(int32 Frame)
 	{
-		return QuatType(-Q.X, -Q.Y, -Q.Z, Q.W);
+		if (!SkeletalMeshComponent)
+		{
+			return false;
+		}
+		
+		UAnimSequence* AnimSequence = Private::GetAnimSequence(SkeletalMeshComponent.Get());
+		if (!AnimSequence)
+		{
+			return false;
+		}
+		AnimSequence->Interpolation = EAnimInterpolationType::Step;
+		const IAnimationDataModel* DataModel = AnimSequence->GetDataModel();
+		if (!DataModel)
+		{
+			return false;
+		}
+		const int32 NumKeys = DataModel->GetNumberOfKeys();
+		if (Frame < 0 || Frame >= NumKeys)
+		{
+			UE_LOG(LogNearestNeighborModel, Warning, TEXT("AnimSequence only has %d keys, but being sampled with key %d"), NumKeys, Frame);
+			return false;
+		}
+
+		Sample(Frame);
+		return true;
 	}
 
-	template<typename QuatType>
-	float Inner(const QuatType& Q1, const QuatType& Q2)
+	void FNearestNeighborGeomCacheSampler::Customize(UAnimSequence* Anim, UGeometryCache* Cache)
 	{
-		return Q1.X * Q2.X + Q1.Y * Q2.Y + Q1.Z * Q2.Z + Q1.W * Q2.W;
+		const TObjectPtr<USkeletalMeshComponent> SkelMeshComp = GetSkeletalMeshComponent();
+		if (SkelMeshComp)
+		{
+			SkelMeshComp->SetAnimationMode(EAnimationMode::AnimationSingleNode);
+			SkelMeshComp->SetAnimation(Anim);
+			SkelMeshComp->SetPosition(0.0f);
+			SkelMeshComp->SetPlayRate(1.0f);
+			SkelMeshComp->Play(false);
+			SkelMeshComp->RefreshBoneTransforms();
+			if (SkelMeshComp->GetAnimInstance())
+			{
+				SkelMeshComp->GetAnimInstance()->GetRequiredBones().SetUseRAWData(true);
+			}
+		}
+
+		const TObjectPtr<UGeometryCacheComponent> GeomCacheComp = GetGeometryCacheComponent();
+		if (GeomCacheComp)
+		{
+			GeomCacheComp->SetGeometryCache(Cache);
+			GeomCacheComp->ResetAnimationTime();
+			GeomCacheComp->SetLooping(false);
+			GeomCacheComp->SetManualTick(true);
+			GeomCacheComp->SetPlaybackSpeed(1.0f);
+			GeomCacheComp->Play();
+		}
+
+		if (Anim && SkelMeshComp && Cache && GeomCacheComp)
+		{
+			GenerateMeshMappings();
+		}
 	}
 
-	template<typename QuatType>
-	QuatType FromVector(const FVector3f& V)
+	void FNearestNeighborGeomCacheSampler::SampleDualQuaternionDeltas(int32 InAnimFrameIndex)
 	{
-		return QuatType(V[0], V[1], V[2], 0);
+		USkeletalMesh* const SkeletalMesh = SkeletalMeshComponent.Get() ? SkeletalMeshComponent->GetSkeletalMeshAsset() : nullptr;
+		UGeometryCache* const GeometryCache = GeometryCacheComponent->GetGeometryCache();
+		if (!SkeletalMeshComponent || !SkeletalMesh || !GeometryCacheComponent || !GeometryCache)
+		{
+			VertexDeltas.Reset(0);
+			return;
+		}
+		const FTransform& AlignmentTransform = Model->GetAlignmentTransform();
+		const FSkeletalMeshModel* ImportedModel = SkeletalMesh->GetImportedModel();
+		constexpr int32 LODIndex = 0;
+		if (!ImportedModel || !ImportedModel->LODModels.IsValidIndex(LODIndex))
+		{
+			VertexDeltas.Reset(0);
+			return;
+		}
+		const FSkeletalMeshLODModel& LODModel = ImportedModel->LODModels[LODIndex];
+		const FSkeletalMeshRenderData* RenderData = SkeletalMesh->GetResourceForRendering();
+		const FSkinWeightVertexBuffer* SkinWeightBuffer = SkeletalMeshComponent->GetSkinWeightBuffer(LODIndex);
+		if (!RenderData || !RenderData->LODRenderData.IsValidIndex(LODIndex) || !SkinWeightBuffer)
+		{
+			VertexDeltas.Reset(0);
+			return;
+		}
+		const FSkeletalMeshLODRenderData& LODRenderData = RenderData->LODRenderData[LODIndex];
+		const TArray<FSkelMeshImportedMeshInfo>& SkelMeshInfos = LODModel.ImportedMeshInfos;
+		if (GeomCacheMeshDatas.Num() != MeshMappings.Num())
+		{
+			GeomCacheMeshDatas.SetNum(MeshMappings.Num());
+		}
+
+		// For all mesh mappings we found.
+		for (int32 MeshMappingIndex = 0; MeshMappingIndex < MeshMappings.Num(); ++MeshMappingIndex)
+		{
+			const UE::MLDeformer::FMLDeformerGeomCacheMeshMapping& MeshMapping = MeshMappings[MeshMappingIndex];
+			const FSkelMeshImportedMeshInfo& MeshInfo = SkelMeshInfos[MeshMapping.MeshIndex];
+			UGeometryCacheTrack* const Track = GeometryCache->Tracks[MeshMapping.TrackIndex];
+		
+			// Sample the mesh data of the geom cache.
+			FGeometryCacheMeshData& GeomCacheMeshData = GeomCacheMeshDatas[MeshMappingIndex];
+			if (!Track->GetMeshDataAtSampleIndex(InAnimFrameIndex, GeomCacheMeshData))
+			{
+				continue;
+			}
+		
+			// Calculate the vertex deltas.
+			for (int32 VertexIndex = 0; VertexIndex < MeshInfo.NumVertices; ++VertexIndex)
+			{
+				const int32 SkinnedVertexIndex = MeshInfo.StartImportedVertex + VertexIndex;
+				const int32 GeomCacheVertexIndex = MeshMapping.SkelMeshToTrackVertexMap[VertexIndex];
+				if (GeomCacheMeshData.Positions.IsValidIndex(GeomCacheVertexIndex))
+				{
+					FVector3f Delta = FVector3f::ZeroVector;
+					const int32 ArrayIndex = 3 * SkinnedVertexIndex;
+					const int32 RenderVertexIndex = MeshMapping.ImportedVertexToRenderVertexMap[VertexIndex];
+					if (RenderVertexIndex != INDEX_NONE)
+					{
+						const FVector3f SkinnedVertexPos = SkinnedVertexPositions[SkinnedVertexIndex];
+						const FVector3f GeomCacheVertexPos = (FVector3f)AlignmentTransform.TransformPosition((FVector)GeomCacheMeshData.Positions[GeomCacheVertexIndex]);
+						const FVector3f WorldDelta = GeomCacheVertexPos - SkinnedVertexPos;
+						Delta = CalcDualQuaternionDelta(RenderVertexIndex, WorldDelta, LODRenderData, *SkinWeightBuffer);
+					}
+					VertexDeltas[ArrayIndex] = Delta.X;
+					VertexDeltas[ArrayIndex + 1] = Delta.Y;
+					VertexDeltas[ArrayIndex + 2] = Delta.Z;
+				}
+			}
+		}
 	}
 
-	template<typename QuatType>
-	FVector3f ToVector(const QuatType& Q)
+	namespace Private
 	{
-		return FVector3f(Q.X, Q.Y, Q.Z);
-	}
+		template<typename QuatType>
+		QuatType Conjugate(const QuatType& Q)
+		{
+			return QuatType(-Q.X, -Q.Y, -Q.Z, Q.W);
+		}
+		
+		template<typename QuatType>
+		float Inner(const QuatType& Q1, const QuatType& Q2)
+		{
+			return Q1.X * Q2.X + Q1.Y * Q2.Y + Q1.Z * Q2.Z + Q1.W * Q2.W;
+		}
+		
+		template<typename QuatType>
+		QuatType FromVector(const FVector3f& V)
+		{
+			return QuatType(V[0], V[1], V[2], 0);
+		}
+		
+		template<typename QuatType>
+		FVector3f ToVector(const QuatType& Q)
+		{
+			return FVector3f(Q.X, Q.Y, Q.Z);
+		}
+	};
 
 	FVector3f FNearestNeighborGeomCacheSampler::CalcDualQuaternionDelta(int32 VertexIndex, const FVector3f& WorldDelta, const FSkeletalMeshLODRenderData& SkelMeshLODData, const FSkinWeightVertexBuffer& SkinWeightBuffer) const
 	{
+		using namespace Private;
 		check(SkeletalMeshComponent);
 		const USkeletalMesh* Mesh = SkeletalMeshComponent->GetSkeletalMeshAsset();
 		check(Mesh);
+		const FSkeletalMeshRenderData* RenderData = Mesh->GetResourceForRendering();
+		constexpr int32 LODIndex = 0;
+		check(RenderData && RenderData->LODRenderData.IsValidIndex(LODIndex));
+		const FSkeletalMeshLODRenderData& LODData = RenderData->LODRenderData[LODIndex];
 
 		// Find the render section, which we need to find the right bone index.
 		int32 SectionIndex = INDEX_NONE;
 		int32 SectionVertexIndex = INDEX_NONE;
-		const FSkeletalMeshLODRenderData& LODData = Mesh->GetResourceForRendering()->LODRenderData[0];
 		LODData.GetSectionFromVertexIndex(VertexIndex, SectionIndex, SectionVertexIndex);
 
 		FQuat4f QuatSum = FQuat4f(0, 0, 0, 0);
@@ -153,7 +248,7 @@ namespace UE::NearestNeighborModel
 				{
 					Sign = Inner(R0, R) < 0 ? -1 : 1;
 				}
-				const float	Weight = Sign * static_cast<float>(WeightByte) * ::UE::AnimationCore::InvMaxRawBoneWeightFloat;
+				const float Weight = Sign * static_cast<float>(WeightByte) * UE::AnimationCore::InvMaxRawBoneWeightFloat;
 				QuatSum += R * Weight;
 			}
 		}
@@ -173,11 +268,11 @@ namespace UE::NearestNeighborModel
 		}
 	}
 
-	uint8 FNearestNeighborGeomCacheSampler::GenerateMeshMappings()
+	EOpFlag FNearestNeighborGeomCacheSampler::GenerateMeshMappings()
 	{
-		uint8 Result = EUpdateResult::SUCCESS;
-		USkeletalMesh* SkeletalMesh = SkeletalMeshComponent.Get() ? SkeletalMeshComponent->GetSkeletalMeshAsset() : nullptr;
-		UGeometryCache* GeometryCache = GeometryCacheComponent.Get() ? GeometryCacheComponent->GetGeometryCache() : nullptr;
+		EOpFlag Result = EOpFlag::Success;
+		USkeletalMesh* const SkeletalMesh = SkeletalMeshComponent.Get() ? SkeletalMeshComponent->GetSkeletalMeshAsset() : nullptr;
+		UGeometryCache* const GeometryCache = GeometryCacheComponent.Get() ? GeometryCacheComponent->GetGeometryCache() : nullptr;
 		if (SkeletalMeshComponent && SkeletalMesh && GeometryCacheComponent && GeometryCache)
 		{
 			TArray<FString> FailedNames;
@@ -185,7 +280,7 @@ namespace UE::NearestNeighborModel
 			GenerateGeomCacheMeshMappings(SkeletalMesh, GeometryCache, MeshMappings, FailedNames, VertexMisMatchNames);
 			if (!FailedNames.IsEmpty() || !VertexMisMatchNames.IsEmpty())
 			{
-				Result |= EUpdateResult::WARNING;
+				Result |= EOpFlag::Warning;
 			}
 			for(int32 i = 0; i < VertexMisMatchNames.Num(); i++)
 			{
@@ -199,93 +294,17 @@ namespace UE::NearestNeighborModel
 		return Result;
 	}
 
-	uint8 FNearestNeighborGeomCacheSampler::CheckMeshMappingsEmpty() const
+	EOpFlag FNearestNeighborGeomCacheSampler::CheckMeshMappingsEmpty() const
 	{
 		if(MeshMappings.IsEmpty())
 		{
 			UE_LOG(LogNearestNeighborModel, Error, TEXT("MeshMappings is empty. Unable to match skeletal mesh with geometry cache."));
-			return EUpdateResult::ERROR;
-		}
-		else
-		{
-			return EUpdateResult::SUCCESS;
-		}
-	}
-
-	bool FNearestNeighborGeomCacheSampler::SetAnimToSample(UAnimSequence& InAnimToSample)
-	{
-		if (SkeletalMeshComponent)
-		{
-			SkeletalMeshComponent->SetAnimationMode(EAnimationMode::AnimationSingleNode);
-			SkeletalMeshComponent->SetAnimation(&InAnimToSample);
-			SkeletalMeshComponent->SetPosition(0.0f);
-			SkeletalMeshComponent->SetPlayRate(1.0f);
-			SkeletalMeshComponent->Play(false);
-			SkeletalMeshComponent->RefreshBoneTransforms();
-			AnimToSample = &InAnimToSample;
-			return true;
-		}
-		else
-		{
-			return false;
-		}
-	}
-
-	int32 FNearestNeighborGeomCacheSampler::GetAnimNumFrames() const
-	{
-		if (AnimToSample)
-		{
-			return AnimToSample->GetDataModel()->GetNumberOfKeys();
+			return EOpFlag::Error;
 		}
 		else
 		{
-			return 0;
-		}
-	}
-
-	bool FNearestNeighborGeomCacheSampler::SampleAnim(int32 Frame)
-	{
-		if (!SkeletalMeshComponent || !AnimToSample)
-		{
-			return false;
-		}
-
-		const int32 NumKeys = AnimToSample->GetDataModel()->GetNumberOfKeys();
-		if (Frame < 0 || Frame >= NumKeys)
-		{
-			UE_LOG(LogNearestNeighborModel, Warning, TEXT("AnimSequence only has %d keys, but being sampled with key %d"), NumKeys, Frame);
-			return false;
-		}
-
-		AnimFrameIndex = Frame;
-		SampleTime = GetTimeAtFrame(Frame);
-		
-		UpdateSkeletalMeshComponent();
-		UpdateBoneRotations();
-		UpdateCurveValues();
-		return true;
-	}
-
-	bool FNearestNeighborGeomCacheSampler::SampleKMeansAnim(const int32 AnimId)
-	{
-		UNearestNeighborModel* NearestNeighborModel = static_cast<UNearestNeighborModel*>(Model);
-		if (NearestNeighborModel && AnimId < NearestNeighborModel->SourceAnims.Num() && SkeletalMeshComponent)
-		{
-			if (const TObjectPtr<UAnimSequence> AnimSequence = NearestNeighborModel->SourceAnims[AnimId])
-			{
-				if (SetAnimToSample(*AnimSequence))
-				{
-					KMeansAnimId = AnimId;
-					return true;
-				}
-			}
+			return EOpFlag::Success;
 		}
-		return false;
-	}
-
-	bool FNearestNeighborGeomCacheSampler::SampleKMeansFrame(const int32 Frame)
-	{
-		return SampleAnim(Frame);
 	}
 
 	TArray<uint32> FNearestNeighborGeomCacheSampler::GetMeshIndexBuffer() const
@@ -297,18 +316,29 @@ namespace UE::NearestNeighborModel
 			return IndexBuffer;
 		}
 
-		USkeletalMesh* Mesh = SkeletalMeshComponent->GetSkeletalMeshAsset();
+		const USkeletalMesh* Mesh = SkeletalMeshComponent->GetSkeletalMeshAsset();
 		if (Mesh == nullptr)
 		{
 			return IndexBuffer;
 		}
 
+		const FSkeletalMeshRenderData* RenderData = Mesh->GetResourceForRendering();
 		constexpr int32 LODIndex = 0;
-		const FSkeletalMeshLODRenderData& SkelMeshLODData = Mesh->GetResourceForRendering()->LODRenderData[LODIndex];
-		SkelMeshLODData.MultiSizeIndexContainer.GetIndexBuffer(IndexBuffer);
+		if(!RenderData || !RenderData->LODRenderData.IsValidIndex(LODIndex))
+		{
+			return IndexBuffer;
+		}
+		const FSkeletalMeshLODRenderData& LODRenderData = Mesh->GetResourceForRendering()->LODRenderData[LODIndex];
 
 		const FSkeletalMeshModel* SkeletalMeshModel = Mesh->GetImportedModel();
-		const TArray<int32>& ImportedVertexNumbers = SkeletalMeshModel->LODModels[LODIndex].MeshToImportVertexMap;
+		if (!SkeletalMeshModel || !SkeletalMeshModel->LODModels.IsValidIndex(LODIndex))
+		{
+			return IndexBuffer;
+		}
+		const FSkeletalMeshLODModel& LODModel = SkeletalMeshModel->LODModels[LODIndex];
+
+		LODRenderData.MultiSizeIndexContainer.GetIndexBuffer(IndexBuffer);
+		TConstArrayView<int32> ImportedVertexNumbers {LODModel.MeshToImportVertexMap};
 		
 		if (ImportedVertexNumbers.Num() > 0)
 		{
@@ -318,6 +348,6 @@ namespace UE::NearestNeighborModel
 			}
 		}
 
-		return MoveTemp(IndexBuffer);
+		return IndexBuffer;
 	}
 };
