Nearestneighbormodel: refactor

diff --git a/Engine/Plugins/Animation/MLDeformer/NearestNeighborModel/Source/NearestNeighborModelEditor/Private/NearestNeighborModelDetails.cpp b/Engine/Plugins/Animation/MLDeformer/NearestNeighborModel/Source/NearestNeighborModelEditor/Private/NearestNeighborModelDetails.cpp
--- a/Engine/Plugins/Animation/MLDeformer/NearestNeighborModel/Source/NearestNeighborModelEditor/Private/NearestNeighborModelDetails.cpp
+++ b/Engine/Plugins/Animation/MLDeformer/NearestNeighborModel/Source/NearestNeighborModelEditor/Private/NearestNeighborModelDetails.cpp
@@ -1,24 +1,20 @@
 // Copyright Epic Games, Inc. All Rights Reserved.
 
 #include "NearestNeighborModelDetails.h"
-#include "NearestNeighborEditorModel.h"
-#include "NearestNeighborModel.h"
-#include "MLDeformerMorphModelDetails.h"
-#include "MLDeformerEditorToolkit.h"
-#include "DetailLayoutBuilder.h"
+
+#include "Components/ExternalMorphSet.h"
+#include "DesktopPlatformModule.h"
 #include "DetailCategoryBuilder.h"
+#include "DetailLayoutBuilder.h"
 #include "DetailWidgetRow.h"
-#include "Engine/SkeletalMesh.h"
 #include "IDetailChildrenBuilder.h"
-#include "IDetailGroup.h"
-#include "IDetailsView.h"
+#include "NearestNeighborEditorModel.h"
+#include "NearestNeighborModel.h"
+#include "NearestNeighborModelSectionCustomization.h"
 #include "PropertyCustomizationHelpers.h"
-#include "Rendering/SkeletalMeshModel.h"
-#include "Rendering/SkeletalMeshLODModel.h"
 #include "SWarningOrErrorBox.h"
 #include "Widgets/Input/SButton.h"
 #include "Widgets/Input/SComboBox.h"
-#include "Widgets/Input/STextComboBox.h"
 #include "Widgets/Layout/SBox.h"
 #include "Widgets/Text/STextBlock.h"
 
@@ -26,8 +22,6 @@
 
 namespace UE::NearestNeighborModel
 {
-	using namespace ::UE::MLDeformer;
-
 	TSharedRef<IDetailCustomization> FNearestNeighborModelDetails::MakeInstance()
 	{
 		return MakeShareable(new FNearestNeighborModelDetails());
@@ -40,21 +34,16 @@ namespace UE::NearestNeighborModel
 			return false;
 		}
 
-		NearestNeighborModel = static_cast<UNearestNeighborModel*>(Model);
-		check(NearestNeighborModel);
-		NearestNeighborEditorModel = static_cast<FNearestNeighborEditorModel*>(EditorModel);
-
-		return (NearestNeighborModel != nullptr && NearestNeighborEditorModel != nullptr);
+		return (GetCastModel() && GetCastEditorModel());
 	}
 
 	void FNearestNeighborModelDetails::CreateCategories()
 	{
 		FMLDeformerMorphModelDetails::CreateCategories();
 
-		FileCacheCategoryBuilder = &DetailLayoutBuilder->EditCategory("File Cache", FText::GetEmpty(), ECategoryPriority::Important);
-		ClothPartCategoryBuilder = &DetailLayoutBuilder->EditCategory("Cloth Parts", FText::GetEmpty(), ECategoryPriority::Important);
-		NearestNeighborCategoryBuilder = &DetailLayoutBuilder->EditCategory("Nearest Neighbors", FText::GetEmpty(), ECategoryPriority::Important);
-		KMeansCategoryBuilder = &DetailLayoutBuilder->EditCategory("KMeans Pose Generator", FText::GetEmpty(), ECategoryPriority::Important);
+		NearestNeighborCategoryBuilder = &DetailLayoutBuilder->EditCategory("Nearest Neighbor Setting", FText::GetEmpty(), ECategoryPriority::Important);
+		SectionsCategoryBuilder = &DetailLayoutBuilder->EditCategory("Sections", FText::GetEmpty(), ECategoryPriority::Important);
+		StatusCategoryBuilder = &DetailLayoutBuilder->EditCategory("Status", FText::GetEmpty(), ECategoryPriority::Important);
 
 		// Add warning in CreateCategories so that the warning appears at the top of the details panel.
 		FDetailWidgetRow& NearestNeighborWarningRow = BaseMeshCategoryBuilder->AddCustomRow(FText::FromString("NearestNeighborWarning"))
@@ -70,251 +59,567 @@ namespace UE::NearestNeighborModel
 			];
 	}
 
-
-	void FNearestNeighborModelDetails::GenerateClothPartElementWidget(TSharedRef<IPropertyHandle> PropertyHandle, int32 ArrayIndex, IDetailChildrenBuilder& ChildrenBuilder)
+	UNearestNeighborModel* FNearestNeighborModelDetails::GetCastModel() const
 	{
-		TSharedPtr<IPropertyHandle> PCACoeffNumPropertyHandle = PropertyHandle->GetChildHandle(TEXT("PCACoeffNum"));
-		TSharedPtr<IPropertyHandle> VertexMapPathHandle = PropertyHandle->GetChildHandle(TEXT("VertexMapPath"));
-
-		typename STextComboBox::FOnTextSelectionChanged SubMeshComboOnSelectionChangedDelegate;
-		SubMeshComboOnSelectionChangedDelegate.BindRaw(this, &FNearestNeighborModelDetails::SubMeshComboSelectionChanged, ArrayIndex);
-		int32 InitMeshIndex = NearestNeighborModel ? NearestNeighborModel->GetPartMeshIndex(ArrayIndex) : 0;
-		InitMeshIndex = InitMeshIndex >= SubMeshNames.Num() ? 0 : InitMeshIndex;
+		return Cast<UNearestNeighborModel>(Model);
+	}
 
-		ChildrenBuilder.AddCustomRow(FText::GetEmpty())
-		[
-			SNew(SVerticalBox)
-			+SVerticalBox::Slot()
-			.AutoHeight()
-			[
-				SNew(SHorizontalBox)
+	FNearestNeighborEditorModel* FNearestNeighborModelDetails::GetCastEditorModel() const
+	{
+		return static_cast<FNearestNeighborEditorModel*>(EditorModel);
+	}
 
-				+SHorizontalBox::Slot()
-				.Padding(5, 2)
-				.FillWidth(0.3f)
+	namespace Private
+	{
+		DECLARE_DELEGATE_RetVal(FText, FGetTextDelegate);
+		void AddTextRow(IDetailCategoryBuilder& CategoryBuilder, const FText& FilterString, const FText& Name, const FGetTextDelegate ValueDelegate)
+		{
+			check(ValueDelegate.IsBound());
+			CategoryBuilder.AddCustomRow(FilterString)
+				.NameContent()
 				[
-					PropertyHandle->CreatePropertyNameWidget()
+					SNew(SHorizontalBox)
+					+SHorizontalBox::Slot()
+					.Padding(2, 2)
+					.VAlign(VAlign_Center)
+					[
+						SNew(STextBlock)
+						.Text(Name)
+						.Font(IDetailLayoutBuilder::GetDetailFont())
+					]
 				]
-
-				+SHorizontalBox::Slot()
-				.Padding(5, 2)
-				.FillWidth(0.7f)
+				.ValueContent()
 				[
-					PropertyHandle->CreatePropertyValueWidget()
-				]
-			]
-
-			+SVerticalBox::Slot()
-			.AutoHeight()
-			[
-				SNew(SHorizontalBox)
+					SNew(SHorizontalBox)
+					+SHorizontalBox::Slot()
+					.Padding(2, 2)
+					.VAlign(VAlign_Center)
+					[
+						SNew(STextBlock)
+						.Text_Lambda([ValueDelegate]
+						{
+							return ValueDelegate.Execute();
+						})
+						.Font(IDetailLayoutBuilder::GetDetailFont())
+					]
+				];
+		}
 
-				+SHorizontalBox::Slot()
-				.Padding(5, 2)
-				.FillWidth(0.3f)
+		void AddFileCacheRow(IDetailCategoryBuilder& CategoryBuilder, const FText& FilterString, const FText& Name, const FGetTextDelegate ValueDelegate, TSharedRef<SWidget>& ValueWidget, TAttribute<bool> bIsEnabled)
+		{
+			check(ValueDelegate.IsBound());
+			CategoryBuilder.AddCustomRow(FilterString)
+				.NameContent()
 				[
-					PCACoeffNumPropertyHandle->CreatePropertyNameWidget()
+					SNew(SHorizontalBox)
+					+SHorizontalBox::Slot()
+					.Padding(2, 2)
+					.VAlign(VAlign_Center)
+					[
+						SNew(STextBlock)
+						.Text(Name)
+						.Font(IDetailLayoutBuilder::GetDetailFont())
+					]
 				]
-
-				+SHorizontalBox::Slot()
-				.Padding(5, 2)
-				.FillWidth(0.7f)
+				.ValueContent()
+				.MinDesiredWidth(300.f)
+				.MaxDesiredWidth(300.f)
 				[
-					PCACoeffNumPropertyHandle->CreatePropertyValueWidget()
+					SNew(SHorizontalBox)
+					+SHorizontalBox::Slot()
+					.Padding(2, 2)
+					.VAlign(VAlign_Center)
+					.FillWidth(0.6f)
+					[
+						SNew(STextBlock)
+						.Text_Lambda([ValueDelegate]
+						{
+							return ValueDelegate.Execute();
+						})
+						.Font(IDetailLayoutBuilder::GetDetailFont())
+					]
+					+SHorizontalBox::Slot()
+					.FillWidth(0.4f)
+					.VAlign(VAlign_Center)
+					[
+						ValueWidget
+					]
 				]
+				.IsEnabled(bIsEnabled);
+		}
 
-			]
+		FDateTime ToLocalTime(const FDateTime& UtcTime)
+		{
+			const FDateTime UtcNow = FDateTime::UtcNow();
+			const FDateTime Now = FDateTime::Now();
+			const FTimespan TimeZoneOffset = Now - UtcNow;
+			return UtcTime + TimeZoneOffset;
+		}
 
-			+SVerticalBox::Slot()
-			.AutoHeight()
-			[
-				SNew(SHorizontalBox)
+		FString ToString(const FDateTime& Time)
+		{
+			return Time.ToString(TEXT("%Y-%m-%d %H:%M:%S"));
+		}
 
-				+SHorizontalBox::Slot()
-				.Padding(5, 2)
-				.FillWidth(0.3f)
-				.VAlign(VAlign_Center)
-				[
-					SNew(STextBlock)
-					.Text(FText::FromString(TEXT("Submesh")))
-					.Font(IDetailLayoutBuilder::GetDetailFont())
-				]
-				+SHorizontalBox::Slot()
-				.Padding(5, 2)
-				.FillWidth(0.7f)
-				[
-					SNew(STextComboBox)
-					.OptionsSource(&SubMeshNames)
-					.OnSelectionChanged(SubMeshComboOnSelectionChangedDelegate)
-					.InitiallySelectedItem(SubMeshNames[InitMeshIndex])
-				]
-			]
-			+SVerticalBox::Slot()
-			.AutoHeight()
-			[
-				SNew(SHorizontalBox)
+		TOptional<FString> OpenDialogAndSelectFolder(const FString& StartingDirectory)
+		{
+			TOptional<FString> None;
+			IDesktopPlatform* DesktopPlatform = FDesktopPlatformModule::Get();
+			if (!DesktopPlatform)
+			{
+				return None;
+			}
 
-				+SHorizontalBox::Slot()
-				.Padding(5, 2)
-				.FillWidth(0.3f)		
-				[
-					VertexMapPathHandle->CreatePropertyNameWidget()
-				]
+			FString FolderName;
+			bool bFolderSelected = DesktopPlatform->OpenDirectoryDialog(
+			    FSlateApplication::Get().FindBestParentWindowHandleForDialogs(nullptr),
+			    TEXT("Choose a directory"),
+			    StartingDirectory,
+			    FolderName
+			);
 
-				+SHorizontalBox::Slot()
-				.Padding(5, 2)
-				.FillWidth(0.7f)
-				[
-					VertexMapPathHandle->CreatePropertyValueWidget()
-				]
-			]
-		];
-	}
+			if (bFolderSelected)
+			{
+				return FolderName;
+			}
+			return None;
+		}
 
-	void FNearestNeighborModelDetails::CustomizeDetails(IDetailLayoutBuilder& DetailBuilder)
-	{
-		// Create all the detail categories and add the properties of the base class.
-		FMLDeformerMorphModelDetails::CustomizeDetails(DetailBuilder);
+		bool DeleteFile(const FString& FilePath)
+		{
+			IPlatformFile& PlatformFile = FPlatformFileManager::Get().GetPlatformFile();
+			if (PlatformFile.FileExists(*FilePath))
+			{
+				return PlatformFile.DeleteFile(*FilePath);
+			}
+			return false; 
+		}
+	};
 
-		if (NearestNeighborModel == nullptr)
+	void FNearestNeighborModelDetails::CustomizeTrainingSettingsCategory() const
+	{
+		if (!TrainingSettingsCategoryBuilder)
 		{
 			return;
 		}
-		
-		// Training settings.
 		TrainingSettingsCategoryBuilder->AddProperty(UNearestNeighborModel::GetInputDimPropertyName());
 		TrainingSettingsCategoryBuilder->AddProperty(UNearestNeighborModel::GetHiddenLayerDimsPropertyName());
 		TrainingSettingsCategoryBuilder->AddProperty(UNearestNeighborModel::GetOutputDimPropertyName());
 		TrainingSettingsCategoryBuilder->AddProperty(UNearestNeighborModel::GetNumEpochsPropertyName());
 		TrainingSettingsCategoryBuilder->AddProperty(UNearestNeighborModel::GetBatchSizePropertyName());
 		TrainingSettingsCategoryBuilder->AddProperty(UNearestNeighborModel::GetLearningRatePropertyName());
-		TrainingSettingsCategoryBuilder->AddProperty(UNearestNeighborModel::GetSavedNetworkSizePropertyName());
-
-		// File cache settings
-		IDetailGroup *Group = &FileCacheCategoryBuilder->AddGroup(TEXT("File Cache"), LOCTEXT("File Cache", "File Cache"), false, true);
-		Group->HeaderProperty(DetailBuilder.GetProperty(UNearestNeighborModel::GetUseFileCachePropertyName()));
-		Group->AddPropertyRow(DetailBuilder.GetProperty(UNearestNeighborModel::GetFileCacheDirectoryPropertyName()));
-		Group->AddPropertyRow(DetailBuilder.GetProperty(UNearestNeighborModel::GetRecomputeDeltasPropertyName()));
-		Group->AddPropertyRow(DetailBuilder.GetProperty(UNearestNeighborModel::GetRecomputePCAPropertyName()));
-
-		BuildSubMeshNames();
-		const int32 MaxPartMeshIndex = NearestNeighborModel->GetMaxPartMeshIndex();
-		if (MaxPartMeshIndex >= SubMeshNames.Num())
-		{
-			UE_LOG(LogNearestNeighborModel, Error, TEXT("Nearest neighbor model was previously created with %d submeshes, but the current skeletal mesh has %d submeshes. Please use the original skeletal mesh or click update to overwrite existing data."), MaxPartMeshIndex + 1, SubMeshNames.Num());
+	}
 
-			NearestNeighborModel->InvalidateClothPartData();
-			AddActionResultText(ClothPartCategoryBuilder, EUpdateResult::ERROR, TEXT("Loading"));
-		}
-		TSharedRef<IPropertyHandle> ClothPartDataPropertyHandle = DetailBuilder.GetProperty(UNearestNeighborModel::GetClothPartEditorDataPropertyName());
-		if (ClothPartDataPropertyHandle->AsArray().IsValid() && !SubMeshNames.IsEmpty())
+	void FNearestNeighborModelDetails::CustomizeNearestNeighborSettingsCategory() const
+	{
+		if (!NearestNeighborCategoryBuilder)
 		{
-			TSharedRef<FDetailArrayBuilder> PropertyBuilder = MakeShared<FDetailArrayBuilder>(ClothPartDataPropertyHandle, true, false, true);
-			PropertyBuilder->OnGenerateArrayElementWidget(FOnGenerateArrayElementWidget::CreateSP(this, &FNearestNeighborModelDetails::GenerateClothPartElementWidget));
-			ClothPartCategoryBuilder->AddCustomBuilder(PropertyBuilder);
+			return;
 		}
+		NearestNeighborCategoryBuilder->AddProperty(UNearestNeighborModel::GetUseDualQuaternionDeltasPropertyName());
+		NearestNeighborCategoryBuilder->AddProperty(UNearestNeighborModel::GetDecayFactorPropertyName());
+		NearestNeighborCategoryBuilder->AddProperty(UNearestNeighborModel::GetNearestNeighborOffsetWeightPropertyName());
+		NearestNeighborCategoryBuilder->AddProperty(UNearestNeighborModel::GetUseRBFPropertyName());
+		NearestNeighborCategoryBuilder->AddProperty(UNearestNeighborModel::GetRBFSigmaPropertyName());
+	}
 
-		// Nearest Neighbor settings
-		NearestNeighborCategoryBuilder->AddProperty(GET_MEMBER_NAME_STRING_CHECKED(UNearestNeighborModel, DecayFactor));
-		NearestNeighborCategoryBuilder->AddProperty(GET_MEMBER_NAME_STRING_CHECKED(UNearestNeighborModel, NearestNeighborOffsetWeight));
-		NearestNeighborCategoryBuilder->AddProperty(UNearestNeighborModel::GetNearestNeighborDataPropertyName());
-
-		KMeansCategoryBuilder->AddProperty(GET_MEMBER_NAME_STRING_CHECKED(UNearestNeighborModel, SourceAnims));
-		KMeansCategoryBuilder->AddProperty(GET_MEMBER_NAME_STRING_CHECKED(UNearestNeighborModel, NumClusters));
-		KMeansCategoryBuilder->AddProperty(GET_MEMBER_NAME_STRING_CHECKED(UNearestNeighborModel, KMeansPartId));
-		KMeansCategoryBuilder->AddCustomRow(FText::FromString(""))
-			.WholeRowContent()
+	void FNearestNeighborModelDetails::CustomizeStatusCategory() const
+	{
+		if (!StatusCategoryBuilder)
+		{
+			return;
+		}
+		Private::AddTextRow(*StatusCategoryBuilder, 
+			LOCTEXT("NetworkStateInfo", "Network State"), 
+			LOCTEXT("NetworkStateName", "Network State"),
+			Private::FGetTextDelegate::CreateLambda([this]
+			{
+				const UNearestNeighborModel* const NearestNeighborModel = GetCastModel();
+				if (NearestNeighborModel && NearestNeighborModel->GetOptimizedNetwork().IsValid())
+				{
+					if (NearestNeighborModel->IsBeforeCustomVersionWasAdded())
+					{
+						return FText::FromString(TEXT("Loaded from previous version."));
+					}
+					else
+					{
+						const FDateTime LocalTime = Private::ToLocalTime(NearestNeighborModel->GetNetworkLastWriteTime());
+						return FText::FromString(FString::Printf(TEXT("Last written: %s"), *Private::ToString(LocalTime)));
+					}
+				}
+				else
+				{
+					return FText::FromString(TEXT("Empty"));
+				}
+			}));
+		
+		Private::AddTextRow(*StatusCategoryBuilder,
+			LOCTEXT("NetworkArchitectureInfo", "Network Architecture"),
+			LOCTEXT("NetworkArchitectureName", "Network Architecture"),
+			Private::FGetTextDelegate::CreateLambda([this]
+			{
+				const UNearestNeighborModel* const NearestNeighborModel = GetCastModel();
+				if (NearestNeighborModel && NearestNeighborModel->GetOptimizedNetwork().IsValid())
+				{
+					if (NearestNeighborModel->IsBeforeCustomVersionWasAdded())
+					{
+						return FText::FromString(TEXT("Loaded from previous version."));
+					}
+					else
+					{
+						return FText::FromString(FString::Printf(TEXT("Last written: %s"),  *NearestNeighborModel->GetNetworkLastWriteArchitectureString()));
+					}
+				}
+				else
+				{
+					return FText::FromString(TEXT("None"));
+				}
+			}));
+		
+		Private::AddTextRow(*StatusCategoryBuilder, 
+			LOCTEXT("MorphTargetsStateInfo", "Morph Targets State"), 
+			LOCTEXT("MorphTargetsStateName", "Morph Targets State"),
+			Private::FGetTextDelegate::CreateLambda([this]
+			{
+				if (const UNearestNeighborModel* const NearestNeighborModel = GetCastModel())
+				{
+					const TSharedPtr<const FExternalMorphSet> MorphSet = NearestNeighborModel->GetMorphTargetSet();
+					if (MorphSet.IsValid() && MorphSet->MorphBuffers.IsMorphResourcesInitialized())
+					{
+						if (NearestNeighborModel->IsBeforeCustomVersionWasAdded())
+						{
+							return FText::FromString(TEXT("Loaded from previous version."));
+						}
+						else
+						{
+							const FDateTime Time = NearestNeighborModel->GetMorphTargetsLastWriteTime();
+							if (Time != FDateTime::MinValue())
+							{
+								const FDateTime LocalTime = Private::ToLocalTime(Time);
+								return FText::FromString(FString::Printf(TEXT("Last written: %s"), *Private::ToString(LocalTime)));
+							}
+						}
+					}
+				}
+				return FText::FromString(TEXT("Empty"));
+			}));
+		
+		Private::AddTextRow(*StatusCategoryBuilder, 
+			LOCTEXT("InferenceStateInfo", "Inference State"),
+			LOCTEXT("InferenceStateName", "Inference State"),
+			Private::FGetTextDelegate::CreateLambda([this]
+			{
+				const UNearestNeighborModel* const NearestNeighborModel = GetCastModel();
+				if (NearestNeighborModel)
+				{
+					if (NearestNeighborModel->IsReadyForInference())
+					{
+						return FText::FromString(TEXT("Valid"));
+					}
+					else if (NearestNeighborModel->IsReadyForTraining())
+					{
+						return FText::FromString(TEXT("Invalid, please Update"));
+					}
+					else
+					{
+						return FText::FromString(TEXT("Invalid, please Train Model"));
+					}
+				}
+				else
+				{
+					return FText::FromString(TEXT("Invalid"));
+				}
+			}));
+		
+		StatusCategoryBuilder->AddCustomRow(FText::FromString("StatusUpdate")).WholeRowContent()
+		[
+			SNew(SBox)
+			.Padding(2, 2)
+			.MaxDesiredWidth(200.f)
 			[
 				SNew(SButton)
-				.Text(FText::FromString("Cluster"))
-				.HAlign(HAlign_Center)
-				.OnClicked_Lambda([this]
+				.Text_Lambda([this]
 				{
-					NearestNeighborEditorModel->KMeansClusterPoses();
-					EditorModel->GetEditor()->GetModelDetailsView()->ForceRefresh();
-					if (NearestNeighborEditorModel->GetKMeansClusterResult() == EUpdateResult::SUCCESS)
+					const UNearestNeighborModel* const NearestNeighborModel = GetCastModel();
+					if (NearestNeighborModel && NearestNeighborModel->IsReadyForInference())
 					{
-						UE_LOG(LogNearestNeighborModel, Display, TEXT("Cluster succeeded."));
+						return FText::FromString(TEXT("Update"));
+					}
+					else
+					{
+						return FText::FromString(TEXT("Update*"));
 					}
-					return FReply::Handled();
 				})
-			];
-				
-		MorphTargetCategoryBuilder->AddProperty(UMLDeformerMorphModel::GetMorphDeltaZeroThresholdPropertyName(), UMLDeformerMorphModel::StaticClass());
-		MorphTargetCategoryBuilder->AddProperty(UMLDeformerMorphModel::GetMorphCompressionLevelPropertyName(), UMLDeformerMorphModel::StaticClass());
-		FText ButtonText = NearestNeighborModel->IsMorphTargetDataValid() ? LOCTEXT("Update", "Update") : LOCTEXT("Update *", "Update *");
-		MorphTargetCategoryBuilder->AddProperty(UNearestNeighborModel::GetMorphDataSizePropertyName());
-		MorphTargetCategoryBuilder->AddCustomRow(FText::FromString(""))
-			.WholeRowContent()
-			[
-				SNew(SButton)
-				.Text(ButtonText)
 				.HAlign(HAlign_Center)
 				.OnClicked_Lambda([this]
 				{
-					NearestNeighborEditorModel->OnMorphTargetUpdate();
-					EditorModel->GetEditor()->GetModelDetailsView()->ForceRefresh();
-					if (NearestNeighborEditorModel->GetMorphTargetUpdateResult() == EUpdateResult::SUCCESS)
+					if (GetCastEditorModel())
 					{
-						UE_LOG(LogNearestNeighborModel, Display, TEXT("Update succeeded."));
+						GetCastEditorModel()->OnUpdateClicked();
 					}
 					return FReply::Handled();
 				})
-			];
-		AddActionResultText(MorphTargetCategoryBuilder, NearestNeighborEditorModel->GetMorphTargetUpdateResult(), TEXT("Update"));
-		AddActionResultText(KMeansCategoryBuilder, NearestNeighborEditorModel->GetKMeansClusterResult(), TEXT("KMeans"));
+				.IsEnabled_Lambda([this]
+				{
+					const UNearestNeighborModel* const NearestNeighborModel = GetCastModel();
+					if (NearestNeighborModel)
+					{
+						return NearestNeighborModel->IsReadyForTraining();
+					}
+					else
+					{
+						return false;
+					}
+				})
+			]
+		];
 	}
 
-	void FNearestNeighborModelDetails::AddActionResultText(IDetailCategoryBuilder* CategoryBuilder, uint8 Result, const FString& ActionName)
+	void FNearestNeighborModelDetails::CustomizeSectionsCategory(IDetailLayoutBuilder& DetailBuilder)
 	{
-		CategoryBuilder->AddCustomRow(FText::FromString("UpdateResultError"))
-			.Visibility((Result & EUpdateResult::ERROR) != 0 ? EVisibility::Visible : EVisibility::Collapsed)
-			.WholeRowContent()
-			[
-				SNew(SBox)
-				.Padding(FMargin(0.0f, 4.0f))
-				[
-					SNew(SWarningOrErrorBox)
-					.MessageStyle(EMessageStyle::Error)
-					.Message(FText::FromString(ActionName + TEXT(" failed with errors. Please check Output Log (LogNearestNeighborModel, LogPython) for details.")))
-				]
-			];
-		CategoryBuilder->AddCustomRow(FText::FromString("UpdateResultWarning"))
-			.Visibility((Result & EUpdateResult::WARNING) != 0 ? EVisibility::Visible : EVisibility::Collapsed)
-			.WholeRowContent()
-			[
-				SNew(SBox)
-				.Padding(FMargin(0.0f, 4.0f))
-				[
-					SNew(SWarningOrErrorBox)
-					.MessageStyle(EMessageStyle::Warning)
-					.Message(FText::FromString(ActionName + TEXT(" finished with warnings. Please check Output Log (LogNearestNeighborModel, LogPython) for details.")))
-				]
-			];
+		TSharedRef<IPropertyHandle> SectionsPropertyHandle = DetailBuilder.GetProperty(UNearestNeighborModel::GetSectionsPropertyName());
+		UNearestNeighborModel* const NearestNeighborModel = GetCastModel();
+		if (SectionsPropertyHandle->AsArray().IsValid() && NearestNeighborModel)
+		{
+			TSharedRef<FDetailArrayBuilder> PropertyBuilder = MakeShared<FDetailArrayBuilder>(SectionsPropertyHandle, true, false, true);
+			PropertyBuilder->OnGenerateArrayElementWidget(FOnGenerateArrayElementWidget::CreateSP(this, &FNearestNeighborModelDetails::GenerateSectionElementWidget));
+			SectionsCategoryBuilder->AddCustomBuilder(PropertyBuilder);
+		}
 	}
 
-	void FNearestNeighborModelDetails::BuildSubMeshNames()
+	void FNearestNeighborModelDetails::CustomizeMorphTargetCategory(IDetailLayoutBuilder& DetailBuilder) const
 	{
-		SubMeshNames.Reset();
-		SubMeshNameMap.Reset();
-		if (NearestNeighborModel && NearestNeighborModel->GetSkeletalMesh() && NearestNeighborModel->GetSkeletalMesh()->GetImportedModel())
+		DetailBuilder.HideProperty(UNearestNeighborModel::GetIncludeMorphTargetNormalsPropertyName());
+	}
+
+	void FNearestNeighborModelDetails::CustomizeFileCacheCategory(IDetailLayoutBuilder& DetailBuilder) const
+	{
+		IDetailCategoryBuilder& Builder = DetailLayoutBuilder->EditCategory("File Cache", FText::GetEmpty(), ECategoryPriority::Important);
+
+		Builder.AddProperty(UNearestNeighborModel::GetUseFileCachePropertyName());
+		Builder.AddCustomRow(FText::FromString("FileCacheWarning"))
+		.Visibility(TAttribute<EVisibility>::Create([this]
 		{
-			const FSkeletalMeshLODModel& LODModel =  NearestNeighborModel->GetSkeletalMesh()->GetImportedModel()->LODModels[0];
-			const TArray<FSkelMeshImportedMeshInfo>& SkelMeshInfos = LODModel.ImportedMeshInfos;
-			for (int32 i = 0; i < SkelMeshInfos.Num(); i++)
+			const UNearestNeighborModel* const NearestNeighborModel = GetCastModel();
+			return NearestNeighborModel && NearestNeighborModel->DoesUseFileCache() ? EVisibility::Visible : EVisibility::Collapsed;
+		}))
+		.WholeRowContent()
+		[
+			SNew(SBox)
+			.HAlign(HAlign_Center)
+			.VAlign(VAlign_Center)
+			[
+				SNew(STextBlock)
+				.Text(LOCTEXT("FileCacheWarning", "Caution: cache files now need to be manually deleted otherwise there can be unexpected results."))
+				.Font(IDetailLayoutBuilder::GetDetailFont())
+				.AutoWrapText(true)
+			]
+		];
+		TSharedPtr<IPropertyHandle> DirectoryHandle = DetailBuilder.GetProperty(UNearestNeighborModel::GetFileCacheDirectoryPropertyName());
+		Builder.AddProperty(DirectoryHandle).CustomWidget()
+		.NameContent()
+		[
+			DirectoryHandle->CreatePropertyNameWidget()
+		]
+		.ValueContent()
+		[
+			SNew(SHorizontalBox)
+			+ SHorizontalBox::Slot()
+			.FillWidth(0.6f)
+			[
+				DirectoryHandle->CreatePropertyValueWidget()
+			]
+			+ SHorizontalBox::Slot()
+			.FillWidth(0.4f)
+			[
+				SNew(SButton)
+				.Text(LOCTEXT("SelectFileCacheFolder", "Select"))
+				.HAlign(HAlign_Center)
+				.OnClicked_Lambda([this]()
+				{
+					UNearestNeighborModel* const NearestNeighborModel = GetCastModel();
+					if (!NearestNeighborModel)
+					{
+						return FReply::Handled();
+					}
+					TOptional<FString> FolderName = Private::OpenDialogAndSelectFolder(NearestNeighborModel->GetFileCacheDirectory());
+					if (FolderName.IsSet())
+					{
+						NearestNeighborModel->SetFileCacheDirectory(FolderName.GetValue());
+						NearestNeighborModel->UpdateFileCache();
+					}
+					return FReply::Handled();
+				})
+			]
+		];
+
+		UNearestNeighborModel* const NearestNeighborModel = GetCastModel();
+		if (!NearestNeighborModel)
+		{
+			return;
+		}
+
+		auto CreateTextDelegate = [NearestNeighborModel](TFunction<TOptional<FDateTime>()> TimestampFunc)
+		{
+			return Private::FGetTextDelegate::CreateLambda([TimestampFunc]
 			{
-				TSharedPtr<FString> StrPtr = MakeShared<FString>(SkelMeshInfos[i].Name.ToString());
-				SubMeshNames.Add(StrPtr);
-				SubMeshNameMap.Add(StrPtr, i);
-			}
+				const TOptional<FDateTime> Timestamp = TimestampFunc();
+				if (Timestamp.IsSet())
+				{
+					return FText::FromString(Private::ToString(Timestamp.GetValue()));
+				}
+				else
+				{
+					return FText::FromString(TEXT("None"));
+				}
+			});
+		};
+
+		auto CreateDeleteButton = [NearestNeighborModel](TFunction<TArray<FString>()> PathsFunc)
+		{
+			return StaticCastSharedRef<SWidget>(SNew(SButton)
+					.Text(LOCTEXT("DeleteFileCache", "Delete"))
+					.HAlign(HAlign_Center)
+					.OnClicked_Lambda([NearestNeighborModel, PathsFunc]()
+					{
+						const TArray<FString> Paths = PathsFunc();
+						for (const FString& Path : Paths)
+						{
+							Private::DeleteFile(Path);
+						}
+						NearestNeighborModel->UpdateFileCache();
+						return FReply::Handled();
+					}));
+		};
+
+		TAttribute<bool> bIsEnabled = TAttribute<bool>::Create(TAttribute<bool>::FGetter::CreateLambda([NearestNeighborModel]
+		{
+			return NearestNeighborModel->DoesUseFileCache();
+		}));
+
+		{
+			TSharedRef<SWidget> Widget = CreateDeleteButton([NearestNeighborModel]()
+			{
+				return NearestNeighborModel->GetCachedDeltasPaths();
+			});
+			Private::AddFileCacheRow(Builder,
+				LOCTEXT("CachedDeltasInfo", "Cached Deltas"),
+				LOCTEXT("CachedDeltasName", "Cached Deltas"),
+				CreateTextDelegate([NearestNeighborModel]
+				{
+					return NearestNeighborModel->GetCachedDeltasTimestamp();
+				}),
+				Widget,
+				bIsEnabled);
 		}
+		{
+			TSharedRef<SWidget> Widget = CreateDeleteButton([NearestNeighborModel]()
+			{
+				return NearestNeighborModel->GetCachedPCAPaths();
+			});
+			Private::AddFileCacheRow(Builder,
+				LOCTEXT("CachedPCAInfo", "Cached PCA"),
+				LOCTEXT("CachedPCAName", "Cached PCA"),
+				CreateTextDelegate([NearestNeighborModel]
+				{
+					return NearestNeighborModel->GetCachedPCATimestamp();
+				}),
+				Widget,
+				bIsEnabled);
+		}
+		{
+			TSharedRef<SWidget> Widget = CreateDeleteButton([NearestNeighborModel]()
+			{
+				return NearestNeighborModel->GetCachedNetworkPaths();
+			});
+			Private::AddFileCacheRow(Builder,
+				LOCTEXT("CachedNetworkInfo", "Cached Network"),
+				LOCTEXT("CachedNetworkName", "Cached Network"),
+				CreateTextDelegate([NearestNeighborModel]
+				{
+					return NearestNeighborModel->GetCachedNetworkTimestamp();
+				}),
+				Widget,
+				bIsEnabled);
+		}
+
+		Builder.AddCustomRow(FText::FromString("FileCacheRefresh")).WholeRowContent()
+		[
+			SNew(SBox)
+			.MaxDesiredWidth(200.f)
+			.Padding(2, 2)
+			[
+				SNew(SButton)
+				.Text(LOCTEXT("RefreshFileCache", "Refresh"))
+				.HAlign(HAlign_Center)
+				.OnClicked_Lambda([NearestNeighborModel]
+				{
+					NearestNeighborModel->UpdateFileCache();
+					return FReply::Handled();
+				})
+				.IsEnabled(bIsEnabled)
+			]
+		];
 	}
 
-	void FNearestNeighborModelDetails::SubMeshComboSelectionChanged(TSharedPtr<FString> InSelectedItem, ESelectInfo::Type SelectInfo, int32 ArrayIndex)
+	void FNearestNeighborModelDetails::CustomizeDetails(IDetailLayoutBuilder& DetailBuilder)
 	{
-		if (NearestNeighborModel)
+		// Create all the detail categories and add the properties of the base class.
+		FMLDeformerMorphModelDetails::CustomizeDetails(DetailBuilder);
+
+		CustomizeTrainingSettingsCategory();
+		CustomizeNearestNeighborSettingsCategory();
+		CustomizeSectionsCategory(DetailBuilder);
+		CustomizeMorphTargetCategory(DetailBuilder);
+		CustomizeStatusCategory();
+		CustomizeFileCacheCategory(DetailBuilder);
+	}
+
+	void FNearestNeighborModelDetails::GenerateSectionElementWidget(TSharedRef<IPropertyHandle> PropertyHandle, int32 ArrayIndex, IDetailChildrenBuilder& ChildrenBuilder)
+	{
+		if (ArrayIndex < GetCastModel()->GetNumSections())
 		{
-			NearestNeighborModel->SetPartMeshIndex(ArrayIndex, SubMeshNameMap[InSelectedItem]);
+			UNearestNeighborModelSection& Section = GetCastModel()->GetSection(ArrayIndex);
+			ChildrenBuilder.AddProperty(PropertyHandle).CustomWidget()
+			.OverrideResetToDefault(FResetToDefaultOverride::Hide())
+			.WholeRowContent()
+			[
+				SNew(SHorizontalBox)
+				+ SHorizontalBox::Slot()
+				[
+					SNew(SNearestNeighborModelSectionWidget)
+					.Section(&Section)
+				]
+				+ SHorizontalBox::Slot()
+				.AutoWidth()
+				.HAlign(HAlign_Left)
+				.VAlign(VAlign_Center)
+				.Padding(4.0f, 1.0f, 0.0f, 1.0f)
+				[ 
+					PropertyCustomizationHelpers::MakeInsertDeleteDuplicateButton(
+						FExecuteAction::CreateLambda([PropertyHandle]()
+						{
+							TSharedPtr<IPropertyHandleArray> ArrayHandle = PropertyHandle->GetParentHandle()->AsArray();
+							check(ArrayHandle.IsValid());
+							const int32 Index = PropertyHandle->GetArrayIndex();
+							ArrayHandle->Insert(Index);
+						}),
+						FExecuteAction::CreateLambda([PropertyHandle]()
+						{
+							TSharedPtr<IPropertyHandleArray> ArrayHandle = PropertyHandle->GetParentHandle()->AsArray();
+							check(ArrayHandle.IsValid());
+							const int32 Index = PropertyHandle->GetArrayIndex();
+							ArrayHandle->DeleteItem(Index);
+						}),
+						FExecuteAction::CreateLambda([PropertyHandle]()
+						{
+							TSharedPtr<IPropertyHandleArray> ArrayHandle = PropertyHandle->GetParentHandle()->AsArray();
+							check(ArrayHandle.IsValid());
+							const int32 Index = PropertyHandle->GetArrayIndex();
+							ArrayHandle->DuplicateItem(Index);
+						}))
+				]
+			];
 		}
 	}
 }	// namespace UE::NearestNeighborModel
diff --git a/Engine/Plugins/Animation/MLDeformer/NearestNeighborModel/Source/NearestNeighborModelEditor/Private/NearestNeighborModelEditorModule.cpp b/Engine/Plugins/Animation/MLDeformer/NearestNeighborModel/Source/NearestNeighborModelEditor/Private/NearestNeighborModelEditorModule.cpp
--- a/Engine/Plugins/Animation/MLDeformer/NearestNeighborModel/Source/NearestNeighborModelEditor/Private/NearestNeighborModelEditorModule.cpp
+++ b/Engine/Plugins/Animation/MLDeformer/NearestNeighborModel/Source/NearestNeighborModelEditor/Private/NearestNeighborModelEditorModule.cpp
@@ -1,12 +1,16 @@
 // Copyright Epic Games, Inc. All Rights Reserved.
 
-#include "NearestNeighborModelVizSettingsDetails.h"
-#include "NearestNeighborModelDetails.h"
-#include "NearestNeighborEditorModel.h"
-#include "NearestNeighborModel.h"
+#include "EditorModeRegistry.h"
 #include "MLDeformerEditorModule.h"
 #include "Modules/ModuleManager.h"
-#include "EditorModeRegistry.h"
+#include "NearestNeighborEditorModel.h"
+#include "NearestNeighborEditorModelActor.h"
+#include "NearestNeighborModel.h"
+#include "NearestNeighborModelSectionCustomization.h"
+#include "NearestNeighborModelDetails.h"
+#include "NearestNeighborModelVizSettingsDetails.h"
+#include "Tools/NearestNeighborKMeansTool.h"
+#include "Tools/NearestNeighborStatsTool.h"
 
 #define LOCTEXT_NAMESPACE "NearestNeighborModelEditorModule"
 
@@ -20,6 +24,9 @@ namespace UE::NearestNeighborModel
 		void StartupModule() override;
 		void ShutdownModule() override;
 		// ~END IModuleInterface overrides.
+	private:
+		TSharedPtr<FNearestNeighborKMeansTool> KMeansTool;
+		TSharedPtr<FNearestNeighborStatsTool> StatsTool;
 	};
 }
 IMPLEMENT_MODULE(UE::NearestNeighborModel::FNearestNeighborModelEditorModule, NearestNeighborModelEditor)
@@ -32,12 +39,21 @@ namespace UE::NearestNeighborModel
 		FPropertyEditorModule& PropertyModule = FModuleManager::LoadModuleChecked<FPropertyEditorModule>("PropertyEditor");
 		PropertyModule.RegisterCustomClassLayout("NearestNeighborModelVizSettings", FOnGetDetailCustomizationInstance::CreateStatic(&FNearestNeighborModelVizSettingsDetails::MakeInstance));
 		PropertyModule.RegisterCustomClassLayout("NearestNeighborModel", FOnGetDetailCustomizationInstance::CreateStatic(&FNearestNeighborModelDetails::MakeInstance));
+		PropertyModule.RegisterCustomClassLayout("NearestNeighborModelSection", FOnGetDetailCustomizationInstance::CreateStatic(&FNearestNeighborModelSectionCustomization::MakeInstance));
 		PropertyModule.NotifyCustomizationModuleChanged();
 
 		// Register our custom ML deformer model to the model registry in the ML Deformer Framework.
+		using ::UE::MLDeformer::FMLDeformerEditorModule;
+		using ::UE::MLDeformer::FMLDeformerEditorModelRegistry;
+		using ::UE::MLDeformer::FOnGetEditorModelInstance;
+		using ::UE::MLDeformer::FMLDeformerEditorModelRegistry;
 		FMLDeformerEditorModule& EditorModule = FModuleManager::LoadModuleChecked<FMLDeformerEditorModule>("MLDeformerFrameworkEditor");
 		FMLDeformerEditorModelRegistry& ModelRegistry = EditorModule.GetModelRegistry();
 		ModelRegistry.RegisterEditorModel(UNearestNeighborModel::StaticClass(), FOnGetEditorModelInstance::CreateStatic(&FNearestNeighborEditorModel::MakeInstance));
+		KMeansTool = MakeShared<FNearestNeighborKMeansTool>();
+		KMeansTool->Register();
+		StatsTool = MakeShared<FNearestNeighborStatsTool>();
+		StatsTool->Register();
 	}
 
 	void FNearestNeighborModelEditorModule::ShutdownModule()
@@ -45,6 +61,8 @@ namespace UE::NearestNeighborModel
 		// Unregister our ML Deformer model.
 		if (FModuleManager::Get().IsModuleLoaded(TEXT("MLDeformerFrameworkEditor")))
 		{
+			using ::UE::MLDeformer::FMLDeformerEditorModule;
+			using ::UE::MLDeformer::FMLDeformerEditorModelRegistry;
 			FMLDeformerEditorModule& EditorModule = FModuleManager::GetModuleChecked<FMLDeformerEditorModule>("MLDeformerFrameworkEditor");
 			FMLDeformerEditorModelRegistry& ModelRegistry = EditorModule.GetModelRegistry();
 			ModelRegistry.UnregisterEditorModel(UNearestNeighborModel::StaticClass());
@@ -56,6 +74,7 @@ namespace UE::NearestNeighborModel
 			FPropertyEditorModule& PropertyModule = FModuleManager::GetModuleChecked<FPropertyEditorModule>(TEXT("PropertyEditor"));
 			PropertyModule.UnregisterCustomClassLayout(TEXT("NearestNeighborModelVizSettings"));
 			PropertyModule.UnregisterCustomClassLayout(TEXT("NearestNeighborModel"));
+			PropertyModule.UnregisterCustomClassLayout(TEXT("NearestNeighborModelSection"));
 			PropertyModule.NotifyCustomizationModuleChanged();
 		}
 	}
diff --git a/Engine/Plugins/Animation/MLDeformer/NearestNeighborModel/Source/NearestNeighborModelEditor/Private/NearestNeighborModelSectionCustomization.cpp b/Engine/Plugins/Animation/MLDeformer/NearestNeighborModel/Source/NearestNeighborModelEditor/Private/NearestNeighborModelSectionCustomization.cpp
new file mode 100644
--- /dev/null
+++ b/Engine/Plugins/Animation/MLDeformer/NearestNeighborModel/Source/NearestNeighborModelEditor/Private/NearestNeighborModelSectionCustomization.cpp
@@ -0,0 +1,153 @@
+// Copyright Epic Games, Inc. All Rights Reserved.
+
+#include "NearestNeighborModelSectionCustomization.h"
+
+#include "DetailWidgetRow.h"
+#include "DetailLayoutBuilder.h"
+#include "MLDeformerEditorModule.h"
+#include "NearestNeighborModel.h"
+#include "NearestNeighborEditorModel.h"
+#include "Widgets/Input/STextComboBox.h"
+
+#define LOCTEXT_NAMESPACE "NearestNeighborModelSectionCustomization"
+
+namespace UE::NearestNeighborModel
+{
+	namespace Private
+	{
+		FNearestNeighborEditorModel* GetEditorModel(const UNearestNeighborModel* Model)
+		{
+			if (!Model)
+			{
+				return nullptr;
+			}
+			using ::UE::MLDeformer::FMLDeformerEditorModule;
+			FMLDeformerEditorModule& EditorModule = FModuleManager::LoadModuleChecked<FMLDeformerEditorModule>("MLDeformerFrameworkEditor");
+			return static_cast<FNearestNeighborEditorModel*>(EditorModule.GetModelRegistry().GetEditorModel(const_cast<UNearestNeighborModel*>(Model)));
+		}
+	};
+
+	void SNearestNeighborModelSectionWidget::Construct(const FArguments& InArgs)
+	{
+		Section = InArgs._Section;
+	
+		FPropertyEditorModule& PropertyModule = FModuleManager::LoadModuleChecked<FPropertyEditorModule>("PropertyEditor");
+	
+		FDetailsViewArgs Args;
+		Args.NameAreaSettings = FDetailsViewArgs::ObjectsUseNameArea;
+		Args.bAllowSearch = false;
+		Args.bShowObjectLabel = false;
+		Args.bShowScrollBar = false;
+		DetailsView = PropertyModule.CreateDetailView(Args);
+		DetailsView->SetObject(Section);
+	
+		this->ChildSlot
+		[
+			SNew(SVerticalBox)
+			+ SVerticalBox::Slot()
+			.AutoHeight()
+			[
+				DetailsView.ToSharedRef()
+			]
+		];
+	}
+
+	void FNearestNeighborModelSectionCustomization::CustomizeDetails(IDetailLayoutBuilder& DetailBuilder)
+	{
+		DetailBuilder.HideCategory("Section Private");
+		IDetailCategoryBuilder& SectionBuilder = DetailBuilder.EditCategory("Section", LOCTEXT("SectionCategory", "Section"));
+		SectionBuilder.AddProperty(UNearestNeighborModelSection::GetNumPCACoeffsPropertyName());
+
+		// Get the selected objects
+		const TArray<TWeakObjectPtr<UObject>>& SelectedObjects = DetailBuilder.GetSelectedObjects();
+		if (SelectedObjects.IsEmpty())
+		{
+			return;
+		}
+		TWeakObjectPtr<UObject> SelectedObject = SelectedObjects[0];
+		if (!SelectedObject.IsValid())
+		{
+			return;
+		}
+		UNearestNeighborModelSection* Section = Cast<UNearestNeighborModelSection>(SelectedObject.Get());
+		if (!Section)
+		{
+			return;
+		}
+
+		const FNearestNeighborEditorModel* EditorModel = Private::GetEditorModel(Section->GetModel());
+		if (!EditorModel)
+		{
+			return;
+		}
+
+		TSharedPtr<IPropertyHandle> VertexMapHandle = DetailBuilder.GetProperty(UNearestNeighborModelSection::GetVertexMapStringPropertyName());
+		if (VertexMapHandle.IsValid())
+		{
+			SectionBuilder.AddProperty(VertexMapHandle).CustomWidget()
+				.OverrideResetToDefault(FResetToDefaultOverride::Hide())
+				.NameContent()
+				[
+					VertexMapHandle->CreatePropertyNameWidget()
+				]
+				.ValueContent()
+				[
+					SNew(SVerticalBox)
+					+ SVerticalBox::Slot()
+					.Padding(2, 2)
+					.AutoHeight()
+					[
+						SNew(STextComboBox)
+						.OptionsSource(EditorModel->GetVertexMapSelector()->GetOptions())
+						.OnSelectionChanged_Lambda([Section, EditorModel](TSharedPtr<FString> InString, ESelectInfo::Type InSelectInfo)
+							{
+								EditorModel->GetVertexMapSelector()->OnSelectionChanged(*Section, InString, InSelectInfo);
+							})
+						.InitiallySelectedItem(EditorModel->GetVertexMapSelector()->GetSelectedItem(*Section))
+						.Font(IDetailLayoutBuilder::GetDetailFont())
+					]
+					+ SVerticalBox::Slot()
+					.AutoHeight()
+					[
+						SNew(SBox)
+						.IsEnabled(TAttribute<bool>::Create(TAttribute<bool>::FGetter::CreateLambda([Section]()
+						{
+							return Section->GetMeshIndex() == INDEX_NONE;
+						})))
+						[
+							VertexMapHandle->CreatePropertyValueWidget()
+						]
+					]
+				];
+		}
+
+		const UNearestNeighborModel* const Model = Section->GetModel();
+		if (!Model)
+		{
+			return;
+		}
+		const FName PosesName = UNearestNeighborModelSection::GetNeighborPosesPropertyName();
+		TSharedPtr<IPropertyHandle> NeighborPosesHandle = DetailBuilder.GetProperty(PosesName);
+		SectionBuilder.AddProperty(PosesName)
+			.IsEnabled(TAttribute<bool>::Create(TAttribute<bool>::FGetter::CreateLambda([Model]()
+				{
+					return Model->IsReadyForTraining();
+				})));
+		const FName MeshesName = UNearestNeighborModelSection::GetNeighborMeshesPropertyName();
+		TSharedPtr<IPropertyHandle> NeighborMeshesHandle = DetailBuilder.GetProperty(MeshesName);
+		SectionBuilder.AddProperty(MeshesName)
+			.IsEnabled(TAttribute<bool>::Create(TAttribute<bool>::FGetter::CreateLambda([Model]()
+				{
+					return Model->IsReadyForTraining();
+				})));
+		const FName FramesName = UNearestNeighborModelSection::GetExcludedFramesPropertyName();
+		TSharedPtr<IPropertyHandle> ExcludedFramesHandle = DetailBuilder.GetProperty(FramesName);
+		SectionBuilder.AddProperty(FramesName)
+			.IsEnabled(TAttribute<bool>::Create(TAttribute<bool>::FGetter::CreateLambda([Model]()
+				{
+					return Model->IsReadyForTraining();
+				})));
+	}
+}	// namespace UE::NearestNeighborModel
+
+#undef LOCTEXT_NAMESPACE
\ No newline at end of file
diff --git a/Engine/Plugins/Animation/MLDeformer/NearestNeighborModel/Source/NearestNeighborModelEditor/Private/NearestNeighborModelStyle.cpp b/Engine/Plugins/Animation/MLDeformer/NearestNeighborModel/Source/NearestNeighborModelEditor/Private/NearestNeighborModelStyle.cpp
--- a/Engine/Plugins/Animation/MLDeformer/NearestNeighborModel/Source/NearestNeighborModelEditor/Private/NearestNeighborModelStyle.cpp
+++ b/Engine/Plugins/Animation/MLDeformer/NearestNeighborModel/Source/NearestNeighborModelEditor/Private/NearestNeighborModelStyle.cpp
@@ -1,13 +1,11 @@
 // Copyright Epic Games, Inc. All Rights Reserved.
 
 #include "NearestNeighborModelStyle.h"
-#include "Styling/SlateStyle.h"
-#include "Styling/SlateStyleRegistry.h"
-#include "Styling/SlateStyleMacros.h"
-#include "Styling/SlateTypes.h"
-#include "Math/Color.h"
+
 #include "Interfaces/IPluginManager.h"
+#include "Math/Color.h"
 #include "Misc/Paths.h"
+#include "Styling/SlateStyleRegistry.h"
 
 namespace UE::NearestNeighborModel
 {
@@ -21,6 +19,7 @@ namespace UE::NearestNeighborModel
 		// Colors and sizes.
 		Set("NearestNeighborModel.NearestNeighborActors.WireframeColor", FLinearColor(1.0f, 1.0f, 0.0f));
 		Set("NearestNeighborModel.NearestNeighborActors.LabelColor", FLinearColor(1.0f, 1.0f, 0.0f));
+		Set("NearestNeighborModel.Verts.VertsColor", FLinearColor(1.0f, 0.5f, 0.0f));
 
 		FSlateStyleRegistry::RegisterSlateStyle(*this);
 	}
diff --git a/Engine/Plugins/Animation/MLDeformer/NearestNeighborModel/Source/NearestNeighborModelEditor/Private/NearestNeighborModelVizSettingsDetails.cpp b/Engine/Plugins/Animation/MLDeformer/NearestNeighborModel/Source/NearestNeighborModelEditor/Private/NearestNeighborModelVizSettingsDetails.cpp
--- a/Engine/Plugins/Animation/MLDeformer/NearestNeighborModel/Source/NearestNeighborModelEditor/Private/NearestNeighborModelVizSettingsDetails.cpp
+++ b/Engine/Plugins/Animation/MLDeformer/NearestNeighborModel/Source/NearestNeighborModelEditor/Private/NearestNeighborModelVizSettingsDetails.cpp
@@ -3,48 +3,80 @@
 #include "NearestNeighborModelVizSettingsDetails.h"
 
 #include "DetailLayoutBuilder.h"
-#include "DetailCategoryBuilder.h"
 #include "DetailWidgetRow.h"
-#include "IDetailsView.h"
 #include "IDetailGroup.h"
 #include "NearestNeighborEditorModel.h"
 #include "NearestNeighborModelVizSettings.h"
-#include "Widgets/Input/SButton.h"
+#include "Widgets/Input/STextComboBox.h"
 
 #define LOCTEXT_NAMESPACE "MLDeformerMorphModelVizSettingsDetails"
 
 namespace UE::NearestNeighborModel
 {
+	TSharedRef<IDetailCustomization> FNearestNeighborModelVizSettingsDetails::MakeInstance()
+	{
+		return MakeShareable(new FNearestNeighborModelVizSettingsDetails());
+	}
+
 	void FNearestNeighborModelVizSettingsDetails::AddAdditionalSettings()
 	{
 		FMLDeformerMorphModelVizSettingsDetails::AddAdditionalSettings();
-		IDetailGroup& NNGroup = LiveSettingsCategory->AddGroup("Nearest Neighbor", LOCTEXT("NearestNeighborLabel", "Nearest Neighbor"), false, true);
-		NNGroup.AddPropertyRow(DetailLayoutBuilder->GetProperty(UNearestNeighborModelVizSettings::GetNearestNeighborActorsOffsetPropertyName(), UNearestNeighborModelVizSettings::StaticClass()));
-		NNGroup.AddPropertyRow(DetailLayoutBuilder->GetProperty(UNearestNeighborModelVizSettings::GetNearestNeighborIdsPropertyName(), UNearestNeighborModelVizSettings::StaticClass()));
+		AddTrainingMeshAdditionalSettings();
+		AddLiveAdditionalSettings();
+	}
 
-		IDetailGroup& ToolsGroup = LiveSettingsCategory->AddGroup("Tools", LOCTEXT("ToolsLabel", "Tools"), false, true);
-		ToolsGroup.AddPropertyRow(DetailLayoutBuilder->GetProperty(UNearestNeighborModelVizSettings::GetNeighborStatsPartIdPropertyName(), UNearestNeighborModelVizSettings::StaticClass()));
-		ToolsGroup.AddWidgetRow()
-		.WholeRowContent()
-		[
-			SNew(SButton)
-			.Text(FText::FromString("Get Neighbor Stats"))
-			.HAlign(HAlign_Center)
-			.OnClicked_Lambda([this]
+	void FNearestNeighborModelVizSettingsDetails::AddTrainingMeshAdditionalSettings()
+	{
+		if (!DetailLayoutBuilder || !TrainingMeshesCategoryBuilder)
+		{
+			return;
+		}
+		const FNearestNeighborEditorModel* const NNEditorModel = GetCastEditorModel();
+		if (!EditorModel)
+		{
+			return;
+		}
+
+		TrainingMeshesCategoryBuilder->AddProperty(GET_MEMBER_NAME_CHECKED(UNearestNeighborModelVizSettings, bDrawVerts), UNearestNeighborModelVizSettings::StaticClass());
+
+		FVertVizSelector* const VertVizSelector = NNEditorModel->GetVertVizSelector();
+		if (!VertVizSelector)
+		{
+			return;
+		}
+
+		TSharedRef<STextComboBox> VertVizComboBox = SNew(STextComboBox)
+			.OptionsSource(VertVizSelector->GetOptions())
+			.InitiallySelectedItem(VertVizSelector->GetSelectedItem())
+			.OnSelectionChanged_Lambda([VertVizSelector](TSharedPtr<FString> InSelectedItem, ESelectInfo::Type SelectInfo)
 			{
-				if (GetCastEditorModel())
-				{
-					GetCastEditorModel()->GetNeighborStats();
-				}
-				return FReply::Handled();
+				VertVizSelector->OnSelectionChanged(InSelectedItem, SelectInfo);
 			})
+			.Font(IDetailLayoutBuilder::GetDetailFont());
+
+		TSharedRef<IPropertyHandle> VertVizHandle = DetailLayoutBuilder->GetProperty(GET_MEMBER_NAME_CHECKED(UNearestNeighborModelVizSettings, VertVizSectionIndex), UNearestNeighborModelVizSettings::StaticClass());
+		TrainingMeshesCategoryBuilder->AddProperty(VertVizHandle).CustomWidget()
+		.NameContent()
+		[
+			VertVizHandle->CreatePropertyNameWidget()
+		]
+		.ValueContent()
+		[
+			VertVizComboBox
 		];
 	}
 
-	FNearestNeighborEditorModel* FNearestNeighborModelVizSettingsDetails::GetCastEditorModel()
+	void FNearestNeighborModelVizSettingsDetails::AddLiveAdditionalSettings()
 	{
-		return static_cast<FNearestNeighborEditorModel*>(EditorModel);
+		IDetailGroup& NNGroup = LiveSettingsCategory->AddGroup("Nearest Neighbor", LOCTEXT("NearestNeighborLabel", "Nearest Neighbor"), false, true);
+		NNGroup.AddPropertyRow(DetailLayoutBuilder->GetProperty(GET_MEMBER_NAME_STRING_CHECKED(UNearestNeighborModelVizSettings, NearestNeighborActorSectionIndex), UNearestNeighborModelVizSettings::StaticClass()));
+		NNGroup.AddPropertyRow(DetailLayoutBuilder->GetProperty(UNearestNeighborModelVizSettings::GetNearestNeighborActorsOffsetPropertyName(), UNearestNeighborModelVizSettings::StaticClass()));
+		NNGroup.AddPropertyRow(DetailLayoutBuilder->GetProperty(UNearestNeighborModelVizSettings::GetNearestNeighborIdsPropertyName(), UNearestNeighborModelVizSettings::StaticClass()));	
 	}
 
+	FNearestNeighborEditorModel* FNearestNeighborModelVizSettingsDetails::GetCastEditorModel() const
+	{
+		return static_cast<FNearestNeighborEditorModel*>(EditorModel);
+	}
 };
 #undef LOCTEXT_NAMESPACE
diff --git a/Engine/Plugins/Animation/MLDeformer/NearestNeighborModel/Source/NearestNeighborModelEditor/Private/NearestNeighborTrainingModel.cpp b/Engine/Plugins/Animation/MLDeformer/NearestNeighborModel/Source/NearestNeighborModelEditor/Private/NearestNeighborTrainingModel.cpp
--- a/Engine/Plugins/Animation/MLDeformer/NearestNeighborModel/Source/NearestNeighborModelEditor/Private/NearestNeighborTrainingModel.cpp
+++ b/Engine/Plugins/Animation/MLDeformer/NearestNeighborModel/Source/NearestNeighborModelEditor/Private/NearestNeighborTrainingModel.cpp
@@ -1,161 +1,172 @@
 // Copyright Epic Games, Inc. All Rights Reserved.
 
 #include "NearestNeighborTrainingModel.h"
+
+#include "Animation/AnimSequence.h"
+#include "Components/SkeletalMeshComponent.h"
+#include "Engine/SkeletalMesh.h"
+#include "NearestNeighborEditorModel.h"
 #include "NearestNeighborModel.h"
 #include "NearestNeighborModelInstance.h"
-#include "NearestNeighborEditorModel.h"
-#include "Animation/AnimSequence.h"
+#include "Tools/NearestNeighborKMeansTool.h"
+#include "Tools/NearestNeighborStatsTool.h"
 
 #define LOCTEXT_NAMESPACE "NearestNeighborTrainingModel"
 
+using UE::NearestNeighborModel::FNearestNeighborEditorModel;
+
 UNearestNeighborTrainingModel::~UNearestNeighborTrainingModel() = default;
 
-void UNearestNeighborTrainingModel::Init(::UE::MLDeformer::FMLDeformerEditorModel* InEditorModel)
+void UNearestNeighborTrainingModel::Init(UE::MLDeformer::FMLDeformerEditorModel* InEditorModel)
 {
-	UMLDeformerTrainingModel::Init(InEditorModel);
+	UMLDeformerGeomCacheTrainingModel::Init(InEditorModel);
 	check(InEditorModel != nullptr);
 	check(InEditorModel->GetModel() != nullptr);
-	NearestNeighborModel = static_cast<UNearestNeighborModel*>(InEditorModel->GetModel());
 }
 
 UNearestNeighborModel* UNearestNeighborTrainingModel::GetNearestNeighborModel() const
 {
-	return Cast<UNearestNeighborModel>(GetModel());
+	return GetCastModel();
 }
 
-UE::NearestNeighborModel::FNearestNeighborEditorModel* UNearestNeighborTrainingModel::GetNearestNeighborEditorModel() const
+int32 UNearestNeighborTrainingModel::GetNumFrames(const UAnimSequence* Anim) const
 {
-	return static_cast<UE::NearestNeighborModel::FNearestNeighborEditorModel*>(EditorModel);
+	return UE::NearestNeighborModel::FHelpers::GetNumFrames(Anim);
 }
 
-const TArray<int32> UNearestNeighborTrainingModel::GetPartVertexMap(const int32 PartId) const
+const USkeleton* UNearestNeighborTrainingModel::GetModelSkeleton(const UMLDeformerModel* Model) const
 {
-	const TArray<uint32>& VertexMap = NearestNeighborModel->PartVertexMap(PartId);
-	return TArray<int32>((int32*)VertexMap.GetData(), VertexMap.Num());
-}
-
-int32 UNearestNeighborTrainingModel::SetSamplerPartData(const int32 PartId)
-{
-	return GetNearestNeighborEditorModel()->SetSamplerPartData(PartId);
+	if (Model)
+	{
+		if (const USkeletalMesh* const SkeletalMesh = Model->GetSkeletalMesh())
+		{
+			return SkeletalMesh->GetSkeleton();
+		}
+	}
+	return nullptr;
 }
 
-int32 UNearestNeighborTrainingModel::GetPartNumNeighbors(const int32 PartId) const
-{
-	return FMath::Min(NearestNeighborModel->GetNumNeighborsFromAnimSequence(PartId), NearestNeighborModel->GetNumNeighborsFromGeometryCache(PartId));
-}
 
-bool UNearestNeighborTrainingModel::SetAnimToSample(UAnimSequence* AnimToSample)
+bool UNearestNeighborTrainingModel::SetCustomSamplerData(UAnimSequence* Anim, UGeometryCache* Cache)
 {
-	if (!NearestNeighborModel || !AnimToSample)
-	{
-		return false;
-	}
-	if (!AnimSampler.Get())
+	if (!CustomSampler.Get())
 	{
-		SetNewAnimSampler();
-		check(AnimSampler.Get());
+		SetNewCustomSampler();
+		if (!CustomSampler.Get() || !CustomSampler->IsInitialized())
+		{
+			UE_LOG(LogNearestNeighborModel, Error, TEXT("Failed to initialize CustomSampler. Please provide at least one AnimSequence in TrainingAnimList"))
+			return false;
+		}
 	}
-	AnimSampler->SetAnimToSample(*AnimToSample);
+	CustomSampler->Customize(Anim, Cache);
 	return true;
 }
 
-
-bool UNearestNeighborTrainingModel::SampleAnim(int32 Frame)
+bool UNearestNeighborTrainingModel::CustomSample(int32 Frame)
 {
-	if (!AnimSampler.Get())
+	if (!CustomSampler.Get())
 	{
 		return false;
 	}
-	if (!AnimSampler->SampleAnim(Frame))
+	if (!CustomSampler->CustomSample(Frame))
 	{
 		return false;
 	}
 
-	SampleBoneRotations = AnimSampler->GetBoneRotations();
+	CustomSamplerBoneRotations = CustomSampler->GetBoneRotations();
+	CustomSamplerDeltas = CustomSampler->GetVertexDeltas();
 	return true;
 }
 
-int32 UNearestNeighborTrainingModel::GetAnimNumFrames() const
+bool UNearestNeighborTrainingModel::SetCustomSamplerDataFromSection(int32 SectionIndex)
 {
-	if (!AnimSampler.Get())
+	if (!CustomSampler.Get())
 	{
-		return 0;
+		SetNewCustomSampler();
+		if (!CustomSampler.Get() || !CustomSampler->IsInitialized())
+		{
+			UE_LOG(LogNearestNeighborModel, Error, TEXT("Failed to initialize CustomSampler. Please provide at least one AnimSequence in TrainingAnimList"))
+			return false;
+		}
 	}
-	return AnimSampler->GetAnimNumFrames();
-}
 
-bool UNearestNeighborTrainingModel::SampleKmeansAnim(const int32 SkeletonId)
-{
-	if (!AnimSampler.Get())
+	const UNearestNeighborModel* const NearestNeighborModel = GetCastModel();
+	if (!NearestNeighborModel)
 	{
-		SetNewAnimSampler();
-		check(AnimSampler.Get());
+		return false;
 	}
-	return AnimSampler->SampleKMeansAnim(SkeletonId);
-}
 
-bool UNearestNeighborTrainingModel::SampleKmeansFrame(const int32 Frame)
-{
-	return SampleAnim(Frame);
-}
+	if (SectionIndex < 0 || SectionIndex >= NearestNeighborModel->GetNumSections())
+	{
+		return false;
+	}
 
-int32 UNearestNeighborTrainingModel::GetKmeansNumAnims() const
-{
-	return NearestNeighborModel->SourceAnims.Num();
+	const UNearestNeighborModel::FSection& Section = NearestNeighborModel->GetSection(SectionIndex);
+	UAnimSequence* AnimSequence = Section.GetMutableNeighborPoses();
+	UGeometryCache* GeometryCache = Section.GetMutableNeighborMeshes();
+	return SetCustomSamplerData(AnimSequence, GeometryCache);
 }
 
-int32 UNearestNeighborTrainingModel::GetKmeansAnimNumFrames(const int32 SkeletonId) const
+namespace UE::NearestNeighborModel::Private
 {
-	if (SkeletonId < NearestNeighborModel->SourceAnims.Num())
+	TArray<float> VectorToFloat(TConstArrayView<FVector3f> VectorArr)
 	{
-		return NearestNeighborModel->SourceAnims[SkeletonId]->GetDataModel()->GetNumberOfFrames();
+		TArray<float> FloatArr;
+		FloatArr.SetNumUninitialized(VectorArr.Num() * 3);
+		for (int32 i = 0; i < VectorArr.Num(); i++)
+		{
+			FloatArr[i * 3] = VectorArr[i].X;
+			FloatArr[i * 3 + 1] = VectorArr[i].Y;
+			FloatArr[i * 3 + 2] = VectorArr[i].Z;
+		}
+		return FloatArr;
 	}
-	return 0;
-}
+};
 
-int32 UNearestNeighborTrainingModel::GetKmeansNumClusters() const
+TArray<float> UNearestNeighborTrainingModel::GetUnskinnedVertexPositions()
 {
-	return NearestNeighborModel->NumClusters;
-}
-
-const TArray<float> UNearestNeighborTrainingModel::GetUnskinnedVertexPositions() const
-{
-	UE::MLDeformer::FMLDeformerSampler* BaseSampler = EditorModel->GetNumTrainingInputAnims() > 0 ? EditorModel->GetSamplerForTrainingAnim(0) : nullptr;
-	UE::NearestNeighborModel::FNearestNeighborGeomCacheSampler* Sampler = static_cast<UE::NearestNeighborModel::FNearestNeighborGeomCacheSampler*>(BaseSampler);
-	if (Sampler == nullptr)
+	TArray<float> Empty;
+	if (!CustomSampler.Get())
 	{
-		return TArray<float>();
+		SetNewCustomSampler();
 	}
-
-	const TArray<FVector3f>& PositionsVec = Sampler->GetUnskinnedVertexPositions();
-	TArray<float> PositionsFloat; 
-	PositionsFloat.SetNumUninitialized(PositionsVec.Num() * 3);
-	for (int32 i = 0; i < PositionsVec.Num(); i++)
+	if (!CustomSampler.Get() || !CustomSampler->IsInitialized())
 	{
-		PositionsFloat[i * 3] = PositionsVec[i].X;
-		PositionsFloat[i * 3 + 1] = PositionsVec[i].Y;
-		PositionsFloat[i * 3 + 2] = PositionsVec[i].Z;
+		UE_LOG(LogNearestNeighborModel, Error, TEXT("Failed to initialize customSampler. Please provide at least one AnimSequence in TrainingAnimList"))
+		return TArray<float>();
 	}
-	return MoveTemp(PositionsFloat);
+	const TArray<FVector3f>& PositionsVec = CustomSampler->GetUnskinnedVertexPositions();
+	return UE::NearestNeighborModel::Private::VectorToFloat(PositionsVec);
 }
 
-const TArray<int32> UNearestNeighborTrainingModel::GetMeshIndexBuffer() const
+TArray<int32> UNearestNeighborTrainingModel::GetMeshIndexBuffer()
 {
-	UE::MLDeformer::FMLDeformerSampler* BaseSampler = EditorModel->GetNumTrainingInputAnims() > 0 ? EditorModel->GetSamplerForTrainingAnim(0) : nullptr;
-	UE::NearestNeighborModel::FNearestNeighborGeomCacheSampler* Sampler = static_cast<UE::NearestNeighborModel::FNearestNeighborGeomCacheSampler*>(BaseSampler);
-	if (Sampler == nullptr)
+	if (!CustomSampler.Get())
+	{
+		SetNewCustomSampler();
+	}
+	if (!CustomSampler.Get() || !CustomSampler->IsInitialized())
 	{
+		UE_LOG(LogNearestNeighborModel, Error, TEXT("Failed to initialize customSampler. Please provide at least one AnimSequence in TrainingAnimList"))
 		return TArray<int32>();
 	}
-	TArray<uint32> UIndexBuffer = Sampler->GetMeshIndexBuffer();
+	TArray<uint32> UIndexBuffer = CustomSampler->GetMeshIndexBuffer();
 	return TArray<int32>((int32*)UIndexBuffer.GetData(), UIndexBuffer.Num());
 }
 
 UNearestNeighborModelInstance* UNearestNeighborTrainingModel::CreateModelInstance()
 {
-	UNearestNeighborModelInstance* ModelInstance = NewObject<UNearestNeighborModelInstance>(this);
-	ModelInstance->SetModel(GetNearestNeighborModel());
-	ModelInstance->Init(nullptr);
+	UNearestNeighborModelInstance* const ModelInstance = NewObject<UNearestNeighborModelInstance>(this);
+	UNearestNeighborModel* const NearestNeighborModel = GetCastModel();
+	if (!NearestNeighborModel)
+	{
+		return nullptr;
+	}
+	ModelInstance->SetModel(NearestNeighborModel);
+	USkeletalMeshComponent* const DummySKC = NewObject<USkeletalMeshComponent>(ModelInstance);
+	DummySKC->SetSkeletalMesh(NearestNeighborModel->GetSkeletalMesh());
+
+	ModelInstance->Init(DummySKC);
 	ModelInstance->PostMLDeformerComponentInit();
 	return ModelInstance;
 }
@@ -165,25 +176,20 @@ void UNearestNeighborTrainingModel::DestroyModelInstance(UNearestNeighborModelIn
 	ModelInstance->ConditionalBeginDestroy();
 }
 
-void UNearestNeighborTrainingModel::SetNewAnimSampler()
+void UNearestNeighborTrainingModel::SetNewCustomSampler()
 {
-	AnimSampler = MakeUnique<UE::NearestNeighborModel::FNearestNeighborGeomCacheSampler>();
-	AnimSampler->Init(EditorModel, 0);
+	CustomSampler = MakeUnique<UE::NearestNeighborModel::FNearestNeighborGeomCacheSampler>();
+	CustomSampler->Init(EditorModel, 0);
 }
 
-const UAnimSequence* UNearestNeighborTrainingModel::GetTestAnim() const
+UNearestNeighborModel* UNearestNeighborTrainingModel::GetCastModel() const
 {
-	if (!EditorModel)
-	{
-		return nullptr;
-	}
-	const UE::NearestNeighborModel::FNearestNeighborEditorModel* NNEditorModel = static_cast<UE::NearestNeighborModel::FNearestNeighborEditorModel*>(EditorModel);
-	const UNearestNeighborModelVizSettings* Viz = NNEditorModel->GetNearestNeighborModelVizSettings();
-	if (!Viz)
-	{
-		return nullptr;
-	}
-	return Viz->GetTestAnimSequence();
+	return Cast<UNearestNeighborModel>(GetModel());
+}
+
+UE::NearestNeighborModel::FNearestNeighborEditorModel* UNearestNeighborTrainingModel::GetCastEditorModel() const
+{
+	return static_cast<UE::NearestNeighborModel::FNearestNeighborEditorModel*>(GetEditorModel());
 }
 
 
diff --git a/Engine/Plugins/Animation/MLDeformer/NearestNeighborModel/Source/NearestNeighborModelEditor/Public/NearestNeighborEditorModel.h b/Engine/Plugins/Animation/MLDeformer/NearestNeighborModel/Source/NearestNeighborModelEditor/Public/NearestNeighborEditorModel.h
--- a/Engine/Plugins/Animation/MLDeformer/NearestNeighborModel/Source/NearestNeighborModelEditor/Public/NearestNeighborEditorModel.h
+++ b/Engine/Plugins/Animation/MLDeformer/NearestNeighborModel/Source/NearestNeighborModelEditor/Public/NearestNeighborEditorModel.h
@@ -4,123 +4,124 @@
 
 #include "CoreTypes.h"
 #include "MLDeformerMorphModelEditorModel.h"
-#include "MLDeformerEditorActor.h"
+#include "NearestNeighborEditorModelActor.h"
+#include "NearestNeighborModelHelpers.h"
+#include "Types/SlateEnums.h"
 
-class UMLDeformerAsset;
-class USkeletalMesh;
-class UGeometryCache;
 class UMLDeformerComponent;
-class UMLDeformerOptimizedNetwork;
+class UNearestNeighborModelSection;
+class USkeletalMesh;
 class UNearestNeighborModel;
 class UNearestNeighborModelInstance;
-class UNearestNeighborTrainingModel;
 class UNearestNeighborModelVizSettings;
 
 namespace UE::NearestNeighborModel
 {
-	using namespace UE::MLDeformer;
-
+	class FVertVizSelector;
 	class FNearestNeighborEditorModelActor;
-	class FNearestNeighborModelDetails;
-	class FNearestNeighborModelSampler;
-	class FNearestNeighborModelVizSettingsDetails;
+	class FVertexMapSelector;
 
 	class NEARESTNEIGHBORMODELEDITOR_API FNearestNeighborEditorModel
 		: public UE::MLDeformer::FMLDeformerMorphModelEditorModel
 	{
 	public:
+		using FSection = UNearestNeighborModelSection;
+		using FMLDeformerSampler = UE::MLDeformer::FMLDeformerSampler;
+
+		virtual ~FNearestNeighborEditorModel();
+
 		// We need to implement this static MakeInstance method.
 		static FMLDeformerEditorModel* MakeInstance();
 
 		// FGCObject overrides.
 		virtual FString GetReferencerName() const override { return TEXT("FNearestNeighborEditorModel"); }
 		// ~END FGCObject overrides.
-	
-		// FMLDeformerGeomCacheEditorModel overrides.
+
+		// FMLDeformerEditorModel overrides.
 		virtual void Init(const InitSettings& Settings) override;
 		virtual TSharedPtr<FMLDeformerSampler> CreateSamplerObject() const override;
 		virtual void CreateActors(const TSharedRef<IPersonaPreviewScene>& InPersonaPreviewScene) override;
 		virtual void Tick(FEditorViewportClient* ViewportClient, float DeltaTime) override;
-		virtual FMLDeformerEditorActor* CreateEditorActor(const FMLDeformerEditorActor::FConstructSettings& Settings) const override;
-
 		virtual void InitInputInfo(UMLDeformerInputInfo* InputInfo) override;
 		virtual ETrainingResult Train() override;
 		virtual bool LoadTrainedNetwork() const override;
-		virtual void OnPropertyChanged(FPropertyChangedEvent& PropertyChangedEvent) override;
 		virtual void OnPostTraining(ETrainingResult TrainingResult, bool bUsePartiallyTrainedWhenAborted) override;
-		
-		UE_DEPRECATED(5.4, "Onnx file no longer used for training network.")
-		virtual FString GetTrainedNetworkOnnxFile() const override;
-		virtual int32 GetNumTrainingFrames() const override;
-		// ~END FMLDeformerGeomCacheEditorModel overrides.
+		virtual FMLDeformerTrainingInputAnim* GetTrainingInputAnim(int32 Index) const override;
+		virtual void UpdateTimelineTrainingAnimList() override;
+		virtual void OnPropertyChanged(FPropertyChangedEvent& PropertyChangedEvent) override;
+		virtual void Render(const FSceneView* View, FViewport* Viewport, FPrimitiveDrawInterface* PDI) override;
+		// ~END FMLDeformerEditorModel overrides.
 
 		// UMLDeformerMorphModelEditorModel overrides.
 		virtual bool IsMorphWeightClampingSupported() const override	{ return false; }	// We already do input clamping, so output clamping really isn't needed.
 		// ~END UMLDeformerMorphModelEditorModel overrides.
 
+		void OnUpdateClicked();
 
-		friend class FNearestNeighborModelSampler;
-		friend class FNearestNeighborModelDetails;
-		friend class ::UNearestNeighborTrainingModel;
-		friend class FNearestNeighborModelVizSettingsDetails;
-
-	private:
-		// Some helpers that cast to this model's variants of some classes.
-		UNearestNeighborModel* GetNearestNeighborModel() const { return static_cast<UNearestNeighborModel*>(Model); }
-		UNearestNeighborModelVizSettings* GetNearestNeighborModelVizSettings() const;
+		FVertexMapSelector* GetVertexMapSelector() const;
+		FVertVizSelector* GetVertVizSelector() const;
 
-		// Recomputes nearest neighbor coeffcients and nearest neighbor vertex offsets. 
-		virtual uint8 UpdateNearestNeighborData();
+	protected:
+		virtual void CreateSamplers() override;
+		virtual bool IsAnimIndexValid(int32 Index) const override; 
 
-		// Temporarily set sampler to use part data in nearest neighbor dataset. This is useful when updating nearest neighbor data.
-		uint8 SetSamplerPartData(const int32 PartI);
+	private:
+		static constexpr int32 DefaultState = INDEX_NONE; 
 
-		// Reset smapler to use the original dataset.
-		void ResetSamplerData();
+		// Some helpers that cast to this model's variants of some classes.
+		UNearestNeighborModel* GetCastModel() const;
+		UNearestNeighborModelVizSettings* GetCastVizSettings() const;
+		UMLDeformerComponent* GetTestMLDeformerComponent() const;
+		USkeletalMeshComponent* GetTestSkeletalMeshComponent() const;
+		UNearestNeighborModelInstance* GetTestNearestNeighborModelInstance() const;
 
-		void UpdateNearestNeighborActors();
-		void KMeansClusterPoses();
-		TPair<UAnimSequence*, uint8> CreateAnimOfClusterCenters(const FString& PackageName, const TArray<int32>& KmeansResults);
-		uint8 GetKMeansClusterResult() const { return KMeansClusterResult; }
+		EOpFlag Update();
+		EOpFlag CheckNetwork();
+		EOpFlag UpdateNearestNeighborData();
+		EOpFlag UpdateMorphDeltas();
+		void ResetMorphTargets();
+		void UpdateNearestNeighborIds();
 
-		uint8 InitMorphTargets();
-		void RefreshMorphTargets();
-		void AddFloatArrayToDeltaArray(const TArray<float>& FloatArr, const TArray<uint32>& VertexMap, TArray<FVector3f>& DeltaArr, int32 DeltaArrayOffset = INDEX_NONE, TOptional<TArray<int32>> OptionalIncludedFrames = TOptional<TArray<int32>>());
+		TUniquePtr<FNearestNeighborEditorModelActor> CreateNearestNeighborActor(UWorld* World) const;
+		void UpdateNearestNeighborActor(FNearestNeighborEditorModelActor& Actor) const;
 
-		int32 GetNumParts();
+		void GetNeighborStats();
 
-		void OnMorphTargetUpdate();
-		uint8 GetMorphTargetUpdateResult() { return MorphTargetUpdateResult; }
-		UMLDeformerComponent* GetTestMLDeformerComponent() const;
-		UMLDeformerModelInstance* GetTestMLDeformerModelInstance() const;
-		void InitTestMLDeformerPreviousWeights();
-		uint8 WarnIfNetworkInvalid();
-
-		virtual void CreateNearestNeighborActors(UWorld* World, int32 StartIndex = 0);
-
-		template<class TrainingModelClass>
-		TrainingModelClass* InitTrainingModel(FMLDeformerEditorModel* EditorModel)
-		{
-			TArray<UClass*> TrainingModels;
-			GetDerivedClasses(TrainingModelClass::StaticClass(), TrainingModels);
-			if (TrainingModels.IsEmpty())
-			{
-				return nullptr;
-			}
-
-			TrainingModelClass* TrainingModel = Cast<TrainingModelClass>(TrainingModels.Last()->GetDefaultObject());
-			TrainingModel->Init(EditorModel);
-			return TrainingModel;
-		}
+		/** This should be only set in CreateActors(). */ 
+		TUniquePtr<FNearestNeighborEditorModelActor> NearestNeighborActor;
 
-		void GetNeighborStats();
+		TUniquePtr<FVertexMapSelector> VertexMapSelector;
+		TUniquePtr<FVertVizSelector> VertVizSelector;
+	};
 
-		// Actors showing the nearest neighbors of the current frame
-		TArray<FNearestNeighborEditorModelActor*> NearestNeighborActors;
+	class FVertexMapSelector
+	{
+	public:
+		void Update(const USkeletalMesh* SkelMesh);
+		TArray<TSharedPtr<FString>>* GetOptions();
+		void OnSelectionChanged(UNearestNeighborModelSection& Section, TSharedPtr<FString> InSelectedItem, ESelectInfo::Type SelectInfo) const;
+		TSharedPtr<FString> GetSelectedItem(const UNearestNeighborModelSection& Section) const;
+		FString GetVertexMapString(const UNearestNeighborModelSection& Section) const;
+		bool IsValid() const;
+	private:
+		void Reset();
+		TArray<TSharedPtr<FString>> Options;
+		TMap<TSharedPtr<FString>, FString> VertexMapStrings;
+		static TSharedPtr<FString> CustomString;
+	};
 
-		UWorld* EditorWorld = nullptr;
-		uint8 MorphTargetUpdateResult = 0;
-		uint8 KMeansClusterResult = 0;
-		int32 NumTrainingFramesOverride = -1;
+	class FVertVizSelector
+	{
+	public:
+		FVertVizSelector(UNearestNeighborModelVizSettings* InSettings);
+		void Update(int32 NumSections);
+		TArray<TSharedPtr<FString>>* GetOptions();
+		void OnSelectionChanged(TSharedPtr<FString> InSelectedItem, ESelectInfo::Type SelectInfo) const;
+		TSharedPtr<FString> GetSelectedItem() const;
+		int32 GetSectionIndex(TSharedPtr<FString> Item) const;
+		void SelectSection(int32 SectionIndex);
+	private:
+		TObjectPtr<UNearestNeighborModelVizSettings> Settings;
+		TArray<TSharedPtr<FString>> Options;
 	};
 }	// namespace UE::NearestNeighborModel
diff --git a/Engine/Plugins/Animation/MLDeformer/NearestNeighborModel/Source/NearestNeighborModelEditor/Public/NearestNeighborEditorModelActor.h b/Engine/Plugins/Animation/MLDeformer/NearestNeighborModel/Source/NearestNeighborModelEditor/Public/NearestNeighborEditorModelActor.h
--- a/Engine/Plugins/Animation/MLDeformer/NearestNeighborModel/Source/NearestNeighborModelEditor/Public/NearestNeighborEditorModelActor.h
+++ b/Engine/Plugins/Animation/MLDeformer/NearestNeighborModel/Source/NearestNeighborModelEditor/Public/NearestNeighborEditorModelActor.h
@@ -3,36 +3,30 @@
 #pragma once
 
 #include "MLDeformerGeomCacheActor.h"
-#include "MLDeformerComponent.h"
 
+class UGeometryCache;
 class UMLDeformerComponent;
-class UGeometryCacheComponent;
 
 namespace UE::NearestNeighborModel
 {
-	class FNearestNeighborEditorModel;
-
 	enum : int32
 	{
 		ActorID_NearestNeighborActors = 6
 	};
 
-	class NEARESTNEIGHBORMODELEDITOR_API FNearestNeighborEditorModelActor
-		: public ::UE::MLDeformer::FMLDeformerGeomCacheActor
+	class FNearestNeighborEditorModelActor
+		: public UE::MLDeformer::FMLDeformerGeomCacheActor
 	{
 	public:
 		FNearestNeighborEditorModelActor(const FConstructSettings& Settings);
+		virtual ~FNearestNeighborEditorModelActor();
 
-		friend class FNearestNeighborEditorModel; 
-
+		void SetGeometryCache(UGeometryCache* InGeometryCache) const;
+		void SetTrackedComponent(const UMLDeformerComponent* InComponent, int32 InSectionIndex);
+		void Tick() const;
+	
 	private:
-		void SetGeometryCacheComponent(UGeometryCacheComponent* Component) { GeomCacheComponent = Component; }
-		UGeometryCacheComponent* GetGeometryCacheComponent() const { return GeomCacheComponent; }
-
-		void InitNearestNeighborActor(const int32 InPartId, const UMLDeformerComponent* InComponent);
-		void TickNearestNeighborActor();
-
-		int32 PartId = INDEX_NONE;
-		const UMLDeformerComponent* MLDeformerComponent = nullptr;
+		int32 SectionIndex = INDEX_NONE;
+		TWeakObjectPtr<const UMLDeformerComponent> TrackedComponent;
 	};
 }	// namespace UE::NearestNeighborModel
diff --git a/Engine/Plugins/Animation/MLDeformer/NearestNeighborModel/Source/NearestNeighborModelEditor/Public/NearestNeighborGeomCacheSampler.h b/Engine/Plugins/Animation/MLDeformer/NearestNeighborModel/Source/NearestNeighborModelEditor/Public/NearestNeighborGeomCacheSampler.h
--- a/Engine/Plugins/Animation/MLDeformer/NearestNeighborModel/Source/NearestNeighborModelEditor/Public/NearestNeighborGeomCacheSampler.h
+++ b/Engine/Plugins/Animation/MLDeformer/NearestNeighborModel/Source/NearestNeighborModelEditor/Public/NearestNeighborGeomCacheSampler.h
@@ -3,9 +3,12 @@
 #pragma once
 
 #include "MLDeformerGeomCacheSampler.h"
+#include "NearestNeighborModelHelpers.h"
 
 class UAnimSequence;
+class UGeometryCache;
 class UNearestNeighborTrainingModel;
+
 namespace UE::NearestNeighborModel
 {
 	class FNearestNeighborEditorModel;
@@ -14,28 +17,34 @@ namespace UE::NearestNeighborModel
 		: public UE::MLDeformer::FMLDeformerGeomCacheSampler
 	{
 	public:
-		virtual ~FNearestNeighborGeomCacheSampler();
 		// FMLDeformerGeomCacheSampler overrides
 		virtual void Sample(int32 InAnimFrameIndex) override;
 		// ~END FMLDeformerGeomCacheSampler overrides
 
-		friend class FNearestNeighborEditorModel;
-		friend class ::UNearestNeighborTrainingModel;
+		/** This will set the SkeletalMeshComponent with custom Anim and GeometryCacheComponent with custom Cache. 
+		 * Using this fuction will stop using Anim and Cache from the EditorModel.
+		 * Using this function will break Sample(). Call CustomSample() instead.
+		 * Do NOT use this function on EditorModel Samplers. 
+		 */ 
+		void Customize(UAnimSequence* Anim, UGeometryCache* Cache = nullptr);
+
+		/**
+		 * Sample function used with Customize
+		 * @return true if sampling is successful
+		 */
+		bool CustomSample(int32 Frame);
+
+		/**
+		 * Get the Mesh Index Buffer of the skeletal mesh object used by the model.
+		 * 
+		 * @return A flattened array of the index buffer.
+		 */
+		TArray<uint32> GetMeshIndexBuffer() const;
 
 	private:
-		virtual void SampleDualQuaternionDeltas(int32 InAnimFrameIndex);
-		virtual bool SampleKMeansAnim(const int32 SkeletonId);
-		virtual bool SampleKMeansFrame(const int32 Frame);
-		int32 GetAnimNumFrames() const;
-		bool SetAnimToSample(UAnimSequence& InAnimToSample);
-		bool SampleAnim(int32 Frame);
-		uint8 GenerateMeshMappings();
-		uint8 CheckMeshMappingsEmpty() const;
-		TArray<uint32> GetMeshIndexBuffer() const;
+		void SampleDualQuaternionDeltas(int32 InAnimFrameIndex);
+		EOpFlag GenerateMeshMappings();
+		EOpFlag CheckMeshMappingsEmpty() const;
 		FVector3f CalcDualQuaternionDelta(int32 VertexIndex, const FVector3f& WorldDelta, const FSkeletalMeshLODRenderData& SkelMeshLODData, const FSkinWeightVertexBuffer& SkinWeightBuffer) const;
-
-		int32 KMeansAnimId;
-
-		TObjectPtr<UAnimSequence> AnimToSample;
 	};
 }
diff --git a/Engine/Plugins/Animation/MLDeformer/NearestNeighborModel/Source/NearestNeighborModelEditor/Public/NearestNeighborModelDetails.h b/Engine/Plugins/Animation/MLDeformer/NearestNeighborModel/Source/NearestNeighborModelEditor/Public/NearestNeighborModelDetails.h
--- a/Engine/Plugins/Animation/MLDeformer/NearestNeighborModel/Source/NearestNeighborModelEditor/Public/NearestNeighborModelDetails.h
+++ b/Engine/Plugins/Animation/MLDeformer/NearestNeighborModel/Source/NearestNeighborModelEditor/Public/NearestNeighborModelDetails.h
@@ -5,10 +5,9 @@
 #include "CoreTypes.h"
 #include "MLDeformerMorphModelDetails.h"
 
-class UNearestNeighborModel;
-class IPropertyHandle;
 class IDetailChildrenBuilder;
-class FDetailWidgetRow;
+class IPropertyHandle;
+class UNearestNeighborModel;
 
 namespace UE::NearestNeighborModel
 {
@@ -31,20 +30,20 @@ namespace UE::NearestNeighborModel
 		// ~END FMLDeformerModelDetails overrides.
 
 	private:
-		UNearestNeighborModel* NearestNeighborModel = nullptr;
-		FNearestNeighborEditorModel* NearestNeighborEditorModel = nullptr;
+		UNearestNeighborModel* GetCastModel() const;
+		FNearestNeighborEditorModel* GetCastEditorModel() const;
 
-		IDetailCategoryBuilder* FileCacheCategoryBuilder = nullptr;
-		IDetailCategoryBuilder* ClothPartCategoryBuilder = nullptr;
-		IDetailCategoryBuilder* NearestNeighborCategoryBuilder = nullptr;
-		IDetailCategoryBuilder* KMeansCategoryBuilder = nullptr;
+		void CustomizeTrainingSettingsCategory() const;
+		void CustomizeNearestNeighborSettingsCategory() const;
+		void CustomizeSectionsCategory(IDetailLayoutBuilder& DetailBuilder);
+		void CustomizeMorphTargetCategory(IDetailLayoutBuilder& DetailBuilder) const;
+		void CustomizeStatusCategory() const;
+		void CustomizeFileCacheCategory(IDetailLayoutBuilder& DetailBuilder) const;
 
-		TArray<TSharedPtr<FString> > SubMeshNames;
-		TMap<TSharedPtr<FString>, int32> SubMeshNameMap;
+		void GenerateSectionElementWidget(TSharedRef<IPropertyHandle> PropertyHandle, int32 ArrayIndex, IDetailChildrenBuilder& ChildrenBuilder);
 
-		void AddActionResultText(IDetailCategoryBuilder* CategoryBuilder, uint8 Result, const FString& ActionName);
-		void GenerateClothPartElementWidget(TSharedRef<IPropertyHandle> PropertyHandle, int32 ArrayIndex, IDetailChildrenBuilder& ChildrenBuilder);
-		void BuildSubMeshNames();
-		void SubMeshComboSelectionChanged(TSharedPtr<FString> InSelectedItem, ESelectInfo::Type SelectInfo, int32 ArrayIndex);
+		IDetailCategoryBuilder* NearestNeighborCategoryBuilder = nullptr;
+		IDetailCategoryBuilder* StatusCategoryBuilder = nullptr;
+		IDetailCategoryBuilder* SectionsCategoryBuilder = nullptr;
 	};
 }	// namespace UE::NearestNeighborModel
diff --git a/Engine/Plugins/Animation/MLDeformer/NearestNeighborModel/Source/NearestNeighborModelEditor/Public/NearestNeighborModelStyle.h b/Engine/Plugins/Animation/MLDeformer/NearestNeighborModel/Source/NearestNeighborModelEditor/Public/NearestNeighborModelStyle.h
--- a/Engine/Plugins/Animation/MLDeformer/NearestNeighborModel/Source/NearestNeighborModelEditor/Public/NearestNeighborModelStyle.h
+++ b/Engine/Plugins/Animation/MLDeformer/NearestNeighborModel/Source/NearestNeighborModelEditor/Public/NearestNeighborModelStyle.h
@@ -11,7 +11,7 @@ namespace UE::NearestNeighborModel
 	{
 	public:
 		FNearestNeighborModelEditorStyle();
-		~FNearestNeighborModelEditorStyle();
+		virtual ~FNearestNeighborModelEditorStyle();
 
 		static FNearestNeighborModelEditorStyle& Get();
 	};
diff --git a/Engine/Plugins/Animation/MLDeformer/NearestNeighborModel/Source/NearestNeighborModelEditor/Public/NearestNeighborModelVizSettingsDetails.h b/Engine/Plugins/Animation/MLDeformer/NearestNeighborModel/Source/NearestNeighborModelEditor/Public/NearestNeighborModelVizSettingsDetails.h
--- a/Engine/Plugins/Animation/MLDeformer/NearestNeighborModel/Source/NearestNeighborModelEditor/Public/NearestNeighborModelVizSettingsDetails.h
+++ b/Engine/Plugins/Animation/MLDeformer/NearestNeighborModel/Source/NearestNeighborModelEditor/Public/NearestNeighborModelVizSettingsDetails.h
@@ -2,24 +2,22 @@
 
 #pragma once
 
-#include "CoreTypes.h"
 #include "MLDeformerMorphModelVizSettingsDetails.h"
 
 namespace UE::NearestNeighborModel
 {
 	class FNearestNeighborEditorModel;
-
 	class NEARESTNEIGHBORMODELEDITOR_API FNearestNeighborModelVizSettingsDetails
-		: public ::UE::MLDeformer::FMLDeformerMorphModelVizSettingsDetails
+		: public UE::MLDeformer::FMLDeformerMorphModelVizSettingsDetails
 	{
 	public:
 		/** Makes a new instance of this detail layout class for a specific detail view requesting it. */
-		static TSharedRef<IDetailCustomization> MakeInstance()
-		{
-			return MakeShareable(new FNearestNeighborModelVizSettingsDetails());
-		}
-		
+		static TSharedRef<IDetailCustomization> MakeInstance();
 		virtual void AddAdditionalSettings() override;
-		FNearestNeighborEditorModel* GetCastEditorModel();
+
+	private:
+		void AddTrainingMeshAdditionalSettings();
+		void AddLiveAdditionalSettings();
+		FNearestNeighborEditorModel* GetCastEditorModel() const;
 	};
 }	// namespace UE::NearestNeighborModel
diff --git a/Engine/Plugins/Animation/MLDeformer/NearestNeighborModel/Source/NearestNeighborModelEditor/Public/NearestNeighborTrainingModel.h b/Engine/Plugins/Animation/MLDeformer/NearestNeighborModel/Source/NearestNeighborModelEditor/Public/NearestNeighborTrainingModel.h
--- a/Engine/Plugins/Animation/MLDeformer/NearestNeighborModel/Source/NearestNeighborModelEditor/Public/NearestNeighborTrainingModel.h
+++ b/Engine/Plugins/Animation/MLDeformer/NearestNeighborModel/Source/NearestNeighborModelEditor/Public/NearestNeighborTrainingModel.h
@@ -2,24 +2,21 @@
 
 #pragma once
 
-#include "CoreTypes.h"
 #include "MLDeformerGeomCacheTrainingModel.h"
 #include "NearestNeighborGeomCacheSampler.h"
-#include "NearestNeighborTrainingModel.generated.h"
-
 
-namespace UE::MLDeformer
-{
-	class FMLDeformerGeomCacheSampler;
-};
+#include "NearestNeighborTrainingModel.generated.h"
 
 namespace UE::NearestNeighborModel
 {
 	class FNearestNeighborEditorModel;
-};
+}
 
 class UAnimSequence;
+class UMLDeformerModel;
+class UNearestNeighborKMeansData;
 class UNearestNeighborModel;
+class UNearestNeighborStatsData;
 class UNearestNeighborModelInstance;
 
 UCLASS(Blueprintable)
@@ -32,72 +29,49 @@ public:
 	virtual ~UNearestNeighborTrainingModel();
 	virtual void Init(UE::MLDeformer::FMLDeformerEditorModel* InEditorModel) override;
 
-	/** Main training function, with implementation in python. */
 	UFUNCTION(BlueprintImplementableEvent, Category = "Python")
 	int32 Train() const;
 
-	friend class UE::NearestNeighborModel::FNearestNeighborEditorModel;
-
-protected:
 	UFUNCTION(BlueprintImplementableEvent, Category = "Python")
 	int32 UpdateNearestNeighborData();
 
 	UFUNCTION(BlueprintImplementableEvent, Category = "Python")
-	int32 KmeansClusterPoses(const int32 PartId) const;
+	int32 KmeansClusterPoses(const UNearestNeighborKMeansData* KMeansData);
 
 	UFUNCTION(BlueprintImplementableEvent, Category = "Python")
-	bool GetNeighborStats(int32 PartId);
+	bool GetNeighborStats(const UNearestNeighborStatsData* StatsData);
 
-	UPROPERTY(EditAnywhere, BlueprintReadWrite, Category = "Training Data")
-	TArray<float> PartSampleDeltas;
-
-	UPROPERTY(EditAnywhere, BlueprintReadWrite, Category = "Training Data")
-	TArray<int32> KmeansResults;
+protected:
+	UPROPERTY(VisibleAnywhere, BlueprintReadOnly, Category = "Python")
+	TArray<float> CustomSamplerBoneRotations;
+	
+	UPROPERTY(VisibleAnywhere, BlueprintReadOnly, Category = "Python")
+	TArray<float> CustomSamplerDeltas;
 
 private:
 	UFUNCTION(BlueprintPure, Category = "Python")
-	UNearestNeighborModel* GetNearestNeighborModel() const;
-
-	UE::NearestNeighborModel::FNearestNeighborEditorModel* GetNearestNeighborEditorModel() const;
-
+	int32 GetNumFrames(const UAnimSequence* Anim) const;
+	
 	UFUNCTION(BlueprintPure, Category = "Python")
-	const TArray<int32> GetPartVertexMap(const int32 PartId) const;
-
-	UFUNCTION(BlueprintCallable, Category = "Python")
-	int32 SetSamplerPartData(const int32 PartId);
-
+	const USkeleton* GetModelSkeleton(const UMLDeformerModel* Model) const;
+	
 	UFUNCTION(BlueprintPure, Category = "Python")
-	int32 GetPartNumNeighbors(const int32 PartId) const;
-
-	UFUNCTION(BlueprintCallable, Category = "Python")
-	bool SetAnimToSample(UAnimSequence* AnimToSample);
+	UNearestNeighborModel* GetNearestNeighborModel() const;
 
 	UFUNCTION(BlueprintCallable, Category = "Python")
-	bool SampleAnim(int32 Frame);
-
-	UFUNCTION(BlueprintPure, Category = "Python")
-	int32 GetAnimNumFrames() const;
+	bool SetCustomSamplerData(UAnimSequence* Anim, UGeometryCache* Cache = nullptr);
 
 	UFUNCTION(BlueprintCallable, Category = "Python")
-	bool SampleKmeansAnim(const int32 SkeletonId);
+	bool CustomSample(int32 Frame);
 
 	UFUNCTION(BlueprintCallable, Category = "Python")
-	bool SampleKmeansFrame(const int32 Frame);
-
-	UFUNCTION(BlueprintPure, Category = "Python")
-	int32 GetKmeansNumAnims() const;
+	bool SetCustomSamplerDataFromSection(int32 SectionIndex);
 
 	UFUNCTION(BlueprintPure, Category = "Python")
-	int32 GetKmeansAnimNumFrames(const int32 SkeletonId) const;
+	TArray<float> GetUnskinnedVertexPositions();
 
 	UFUNCTION(BlueprintPure, Category = "Python")
-	int32 GetKmeansNumClusters() const;
-
-	UFUNCTION(BlueprintPure, Category = "Python")
-	const TArray<float> GetUnskinnedVertexPositions() const;
-
-	UFUNCTION(BlueprintPure, Category = "Python")
-	const TArray<int32> GetMeshIndexBuffer() const;
+	TArray<int32> GetMeshIndexBuffer();
 
 	UFUNCTION(BlueprintCallable, Category = "Python")
 	UNearestNeighborModelInstance* CreateModelInstance();
@@ -105,12 +79,9 @@ private:
 	UFUNCTION(BlueprintCallable, Category = "Python")
 	void DestroyModelInstance(UNearestNeighborModelInstance* ModelInstance);
 
-	UFUNCTION(BlueprintCallable, Category = "Python")
-	const UAnimSequence* GetTestAnim() const;
-
-	void SetNewAnimSampler();
-
-	UNearestNeighborModel* NearestNeighborModel = nullptr;
+	void SetNewCustomSampler();
+	UNearestNeighborModel* GetCastModel() const;
+	UE::NearestNeighborModel::FNearestNeighborEditorModel* GetCastEditorModel() const;
 
-	TUniquePtr<UE::NearestNeighborModel::FNearestNeighborGeomCacheSampler> AnimSampler;
+	TUniquePtr<UE::NearestNeighborModel::FNearestNeighborGeomCacheSampler> CustomSampler;
 };
