NearestNeighborModel: fix linux compile error

diff --git a/Engine/Plugins/Animation/MLDeformer/NearestNeighborModel/Source/NearestNeighborModelEditor/Private/NearestNeighborEditorModel.cpp b/Engine/Plugins/Animation/MLDeformer/NearestNeighborModel/Source/NearestNeighborModelEditor/Private/NearestNeighborEditorModel.cpp
--- a/Engine/Plugins/Animation/MLDeformer/NearestNeighborModel/Source/NearestNeighborModelEditor/Private/NearestNeighborEditorModel.cpp
+++ b/Engine/Plugins/Animation/MLDeformer/NearestNeighborModel/Source/NearestNeighborModelEditor/Private/NearestNeighborEditorModel.cpp
@@ -612,7 +612,6 @@ namespace UE::NearestNeighborModel
 		const int32 TotalNumNeighbors = NearestNeighborModel->GetTotalNumNeighbors();
 		Deltas.SetNumZeroed((1 + TotalNumPCACoeffs + TotalNumNeighbors) * NumBaseMeshVerts);
 	
-		using FSection = UNearestNeighborModel::FSection;
 		using Private::AddFloatArrayToDeltaArray;
 		int32 MorphOffset = 1;
 		for (int32 SectionIndex = 0; SectionIndex < NearestNeighborModel->GetNumSections(); ++SectionIndex)
