NearestNeighborModel: fix NonUnity Compile UnrealEditor Mac error

diff --git a/Engine/Plugins/Animation/MLDeformer/NearestNeighborModel/Source/NearestNeighborModel/Public/NearestNeighborModelHelpers.h b/Engine/Plugins/Animation/MLDeformer/NearestNeighborModel/Source/NearestNeighborModel/Public/NearestNeighborModelHelpers.h
--- a/Engine/Plugins/Animation/MLDeformer/NearestNeighborModel/Source/NearestNeighborModel/Public/NearestNeighborModelHelpers.h
+++ b/Engine/Plugins/Animation/MLDeformer/NearestNeighborModel/Source/NearestNeighborModel/Public/NearestNeighborModelHelpers.h
@@ -3,6 +3,7 @@
 
 #include "CoreTypes.h"
 #include "Logging/LogMacros.h"
+#include "UObject/Class.h"
 
 class UAnimSequence;
 class UGeometryCache;
