NearestNeighborModel: remove GetDerivedCDO. Use NewDerivedObject instead.
Fix a crash when setting the training anim to None.
Fix a crash of nearest neighbor model when initializing training model.
Fix a bug where morph target is not created

diff --git a/Engine/Plugins/Animation/MLDeformer/MLDeformerFramework/Source/MLDeformerFrameworkEditor/Private/MLDeformerGeomCacheTrainingModel.cpp b/Engine/Plugins/Animation/MLDeformer/MLDeformerFramework/Source/MLDeformerFrameworkEditor/Private/MLDeformerGeomCacheTrainingModel.cpp
--- a/Engine/Plugins/Animation/MLDeformer/MLDeformerFramework/Source/MLDeformerFrameworkEditor/Private/MLDeformerGeomCacheTrainingModel.cpp
+++ b/Engine/Plugins/Animation/MLDeformer/MLDeformerFramework/Source/MLDeformerFrameworkEditor/Private/MLDeformerGeomCacheTrainingModel.cpp
@@ -25,10 +25,10 @@ bool UMLDeformerGeomCacheTrainingModel::FindNextAnimToSample(int32& OutNextAnimI
 	const int32 NumInputAnims = EditorModel->GetNumTrainingInputAnims();
 	while (NumTries < NumInputAnims)	// Try all input animations at worst case.
 	{
-		const FMLDeformerGeomCacheTrainingInputAnim& InputAnim = GeomCacheModel->GetTrainingInputAnims()[AnimIndex];
-		if (InputAnim.IsEnabled())
+		const FMLDeformerTrainingInputAnim* InputAnim = EditorModel->GetTrainingInputAnim(AnimIndex);
+		if (InputAnim && InputAnim->IsEnabled())
 		{
-			if (NumTimesSampled[AnimIndex] < InputAnim.GetNumFramesToSample())
+			if (NumTimesSampled[AnimIndex] < InputAnim->GetNumFramesToSample())
 			{
 				OutNextAnimIndex = AnimIndex;
 				return true;
diff --git a/Engine/Plugins/Animation/MLDeformer/MLDeformerFramework/Source/MLDeformerFrameworkEditor/Private/MLDeformerMorphModelEditorModel.cpp b/Engine/Plugins/Animation/MLDeformer/MLDeformerFramework/Source/MLDeformerFrameworkEditor/Private/MLDeformerMorphModelEditorModel.cpp
--- a/Engine/Plugins/Animation/MLDeformer/MLDeformerFramework/Source/MLDeformerFrameworkEditor/Private/MLDeformerMorphModelEditorModel.cpp
+++ b/Engine/Plugins/Animation/MLDeformer/MLDeformerFramework/Source/MLDeformerFrameworkEditor/Private/MLDeformerMorphModelEditorModel.cpp
@@ -411,6 +411,11 @@ namespace UE::MLDeformer
 		bool bInvertMaskChannel)
 	{
 		FMLDeformerSampler* Sampler = GetSamplerForActiveAnim();
+		// In case when ActiveTrainingInputAnimIndex == INDEX_NONE but we have a sampler (e.g. testing), use that.
+		if (Sampler == nullptr)
+		{
+			Sampler = GetSamplerForTrainingAnim(0);
+		}
 
 		OutMorphTargets.Reset();
 		if (Deltas.IsEmpty() || Sampler == nullptr)
diff --git a/Engine/Plugins/Animation/MLDeformer/MLDeformerFramework/Source/MLDeformerFrameworkEditor/Private/MLDeformerSampler.cpp b/Engine/Plugins/Animation/MLDeformer/MLDeformerFramework/Source/MLDeformerFrameworkEditor/Private/MLDeformerSampler.cpp
--- a/Engine/Plugins/Animation/MLDeformer/MLDeformerFramework/Source/MLDeformerFrameworkEditor/Private/MLDeformerSampler.cpp
+++ b/Engine/Plugins/Animation/MLDeformer/MLDeformerFramework/Source/MLDeformerFrameworkEditor/Private/MLDeformerSampler.cpp
@@ -106,6 +106,11 @@ namespace UE::MLDeformer
 
 	void FMLDeformerSampler::Sample(int32 InAnimFrameIndex)
 	{
+		FMLDeformerTrainingInputAnim* TrainingInputAnim = EditorModel->GetTrainingInputAnim(AnimIndex);
+		if (!TrainingInputAnim || !TrainingInputAnim->IsValid())
+		{
+			return;
+		}
 		UAnimSequence* TrainingAnimSequence = EditorModel->GetTrainingInputAnim(AnimIndex)->GetAnimSequence();
 		check(TrainingAnimSequence);
 		const EAnimInterpolationType InterpolationTypeBackup = TrainingAnimSequence->Interpolation;
diff --git a/Engine/Plugins/Animation/MLDeformer/NearestNeighborModel/Source/NearestNeighborModel/Public/NearestNeighborModelHelpers.h b/Engine/Plugins/Animation/MLDeformer/NearestNeighborModel/Source/NearestNeighborModel/Public/NearestNeighborModelHelpers.h
--- a/Engine/Plugins/Animation/MLDeformer/NearestNeighborModel/Source/NearestNeighborModel/Public/NearestNeighborModelHelpers.h
+++ b/Engine/Plugins/Animation/MLDeformer/NearestNeighborModel/Source/NearestNeighborModel/Public/NearestNeighborModelHelpers.h
@@ -3,7 +3,6 @@
 
 #include "CoreTypes.h"
 #include "Logging/LogMacros.h"
-#include "UObject/Class.h"
 
 class UAnimSequence;
 class UGeometryCache;
@@ -21,7 +20,7 @@ namespace UE::NearestNeighborModel
 		static int32 GetNumFrames(const UGeometryCache* GeometryCache);
 		
 		template<class T>
-		static T* GetDerivedCDO()
+		static T* NewDerivedObject()
 		{
 			TArray<UClass*> Classes;
 			GetDerivedClasses(T::StaticClass(), Classes);
@@ -29,9 +28,7 @@ namespace UE::NearestNeighborModel
 			{
 				return nullptr;
 			}
-				
-			T* Object = Cast<T>(Classes.Last()->GetDefaultObject());
-			return Object;
+			return NewObject<T>(GetTransientPackage(), Classes.Last());
 		}
 #endif
 	
diff --git a/Engine/Plugins/Animation/MLDeformer/NearestNeighborModel/Source/NearestNeighborModelEditor/Private/NearestNeighborEditorHelpers.cpp b/Engine/Plugins/Animation/MLDeformer/NearestNeighborModel/Source/NearestNeighborModelEditor/Private/NearestNeighborEditorHelpers.cpp
--- a/Engine/Plugins/Animation/MLDeformer/NearestNeighborModel/Source/NearestNeighborModelEditor/Private/NearestNeighborEditorHelpers.cpp
+++ b/Engine/Plugins/Animation/MLDeformer/NearestNeighborModel/Source/NearestNeighborModelEditor/Private/NearestNeighborEditorHelpers.cpp
@@ -118,7 +118,7 @@ bool UNearestNeighborAnimStream::AppendFrames(const UAnimSequence* Anim, TArray<
 		check(RotKeys[Index].Num() == CurrentFrames);
 		RotKeys[Index].SetNum(CurrentFrames + NewFrames);
 		check(ScaleKeys[Index].Num() == CurrentFrames);
-		ScaleKeys[Index].Reserve(CurrentFrames + NewFrames);
+		ScaleKeys[Index].SetNum(CurrentFrames + NewFrames);
 	}
 	for (int32 Index = 0; Index < NewFrames; ++Index)
 	{
diff --git a/Engine/Plugins/Animation/MLDeformer/NearestNeighborModel/Source/NearestNeighborModelEditor/Private/NearestNeighborEditorModel.cpp b/Engine/Plugins/Animation/MLDeformer/NearestNeighborModel/Source/NearestNeighborModelEditor/Private/NearestNeighborEditorModel.cpp
--- a/Engine/Plugins/Animation/MLDeformer/NearestNeighborModel/Source/NearestNeighborModelEditor/Private/NearestNeighborEditorModel.cpp
+++ b/Engine/Plugins/Animation/MLDeformer/NearestNeighborModel/Source/NearestNeighborModelEditor/Private/NearestNeighborEditorModel.cpp
@@ -518,7 +518,16 @@ namespace UE::NearestNeighborModel
 			return UpdateResult;
 		}
 		InitEngineMorphTargets(NearestNeighborModel->GetMorphTargetDeltas());
-		NearestNeighborModel->UpdateMorphTargetsLastWriteTime();
+		const TSharedPtr<const FExternalMorphSet> MorphSet = NearestNeighborModel->GetMorphTargetSet();
+		if (MorphSet.IsValid() && MorphSet->MorphBuffers.IsMorphResourcesInitialized())
+		{
+			NearestNeighborModel->UpdateMorphTargetsLastWriteTime();
+		}
+		else
+		{
+			UE_LOG(LogNearestNeighborModel, Error, TEXT("Morph target set is empty"));
+			UpdateResult |= EOpFlag::Error;
+		}
 		SetDefaultDeformerGraphIfNeeded();
 		return UpdateResult;
 	}
@@ -571,7 +580,7 @@ namespace UE::NearestNeighborModel
 			UE_LOG(LogNearestNeighborModel, Error, TEXT("Network is not trained. Nearest neighbor data cannot be updated."));
 			return EOpFlag::Error;
 		}
-		UNearestNeighborTrainingModel *TrainingModel = FHelpers::GetDerivedCDO<UNearestNeighborTrainingModel>();
+		UNearestNeighborTrainingModel *TrainingModel = FHelpers::NewDerivedObject<UNearestNeighborTrainingModel>();
 		if (!TrainingModel)
 		{
 			return EOpFlag::Error;
@@ -680,7 +689,7 @@ namespace UE::NearestNeighborModel
 	
 		if (Deltas.Num() == 0)
 		{
-			UE_LOG(LogNearestNeighborModel, Error, TEXT("All cloth parts are empty. No morph targets are generated."));
+			UE_LOG(LogNearestNeighborModel, Error, TEXT("All sections are empty. No morph targets are generated."));
 			return EOpFlag::Error;
 		}
 
diff --git a/Engine/Plugins/Animation/MLDeformer/NearestNeighborModel/Source/NearestNeighborModelEditor/Private/Tools/NearestNeighborKMeansTool.cpp b/Engine/Plugins/Animation/MLDeformer/NearestNeighborModel/Source/NearestNeighborModelEditor/Private/Tools/NearestNeighborKMeansTool.cpp
--- a/Engine/Plugins/Animation/MLDeformer/NearestNeighborModel/Source/NearestNeighborModelEditor/Private/Tools/NearestNeighborKMeansTool.cpp
+++ b/Engine/Plugins/Animation/MLDeformer/NearestNeighborModel/Source/NearestNeighborModelEditor/Private/Tools/NearestNeighborKMeansTool.cpp
@@ -147,7 +147,7 @@ namespace UE::NearestNeighborModel
 					{
 						if (UNearestNeighborKMeansData* KMeansData = Cast<UNearestNeighborKMeansData>(&Data))
 						{
-							UNearestNeighborTrainingModel* TrainingModel = FHelpers::GetDerivedCDO<UNearestNeighborTrainingModel>();
+							UNearestNeighborTrainingModel* TrainingModel = FHelpers::NewDerivedObject<UNearestNeighborTrainingModel>();
 							TrainingModel->Init(EditorModel.Get());
 							const int32 ResultInt = TrainingModel->KmeansClusterPoses(KMeansData);
 							const EOpFlag Result = ToOpFlag(ResultInt);
diff --git a/Engine/Plugins/Animation/MLDeformer/NearestNeighborModel/Source/NearestNeighborModelEditor/Private/Tools/NearestNeighborStatsTool.cpp b/Engine/Plugins/Animation/MLDeformer/NearestNeighborModel/Source/NearestNeighborModelEditor/Private/Tools/NearestNeighborStatsTool.cpp
--- a/Engine/Plugins/Animation/MLDeformer/NearestNeighborModel/Source/NearestNeighborModelEditor/Private/Tools/NearestNeighborStatsTool.cpp
+++ b/Engine/Plugins/Animation/MLDeformer/NearestNeighborModel/Source/NearestNeighborModelEditor/Private/Tools/NearestNeighborStatsTool.cpp
@@ -59,7 +59,7 @@ namespace UE::NearestNeighborModel
 					{
 						if (UNearestNeighborStatsData* StatsData = Cast<UNearestNeighborStatsData>(&Data))
 						{
-							UNearestNeighborTrainingModel* TrainingModel = FHelpers::GetDerivedCDO<UNearestNeighborTrainingModel>();
+							UNearestNeighborTrainingModel* TrainingModel = FHelpers::NewDerivedObject<UNearestNeighborTrainingModel>();
 							TrainingModel->Init(EditorModel.Get());
 							const int32 ResultInt = TrainingModel->GetNeighborStats(StatsData);
 							const EOpFlag Result = ToOpFlag(ResultInt);
