NearestNeighborModel: allow multiple modes of weight map selection

diff --git a/Engine/Plugins/Animation/MLDeformer/NearestNeighborModel/Source/NearestNeighborModel/NearestNeighborModel.Build.cs b/Engine/Plugins/Animation/MLDeformer/NearestNeighborModel/Source/NearestNeighborModel/NearestNeighborModel.Build.cs
--- a/Engine/Plugins/Animation/MLDeformer/NearestNeighborModel/Source/NearestNeighborModel/NearestNeighborModel.Build.cs
+++ b/Engine/Plugins/Animation/MLDeformer/NearestNeighborModel/Source/NearestNeighborModel/NearestNeighborModel.Build.cs
@@ -43,6 +43,7 @@ public class NearestNeighborModel : ModuleRules
 				"RenderCore",
 				"RHI",
 				"MLDeformerFramework",
+				"MeshDescription",
 				"NNE",
 				"NNERuntimeBasicCpu",
 			}
diff --git a/Engine/Plugins/Animation/MLDeformer/NearestNeighborModel/Source/NearestNeighborModel/Private/NearestNeighborModel.cpp b/Engine/Plugins/Animation/MLDeformer/NearestNeighborModel/Source/NearestNeighborModel/Private/NearestNeighborModel.cpp
--- a/Engine/Plugins/Animation/MLDeformer/NearestNeighborModel/Source/NearestNeighborModel/Private/NearestNeighborModel.cpp
+++ b/Engine/Plugins/Animation/MLDeformer/NearestNeighborModel/Source/NearestNeighborModel/Private/NearestNeighborModel.cpp
@@ -2,6 +2,7 @@
 
 #include "NearestNeighborModel.h"
 #include "Algo/Copy.h"
+#include "BoneWeights.h"
 #include "Components/ExternalMorphSet.h"
 #include "Components/SkinnedMeshComponent.h"
 #include "Engine/SkeletalMesh.h"
@@ -9,6 +10,7 @@
 #include "HAL/PlatformFileManager.h"
 #include "Interfaces/IPluginManager.h"
 #include "Internationalization/Regex.h"
+#include "MeshDescription.h"
 #include "Misc/FileHelper.h"
 #include "Misc/Optional.h"
 #include "MLDeformerComponent.h"
@@ -281,6 +283,24 @@ namespace UE::NearestNeighborModel
 			}
 			return MeshVertRanges;
 		}
+
+		TArray<FName> GetVertexFloatAttributeNames(const USkeletalMesh& SkelMesh)
+		{
+			FMeshDescription MeshDescription;
+			constexpr int32 LODIndex = 0;
+			SkelMesh.GetMeshDescription(LODIndex, MeshDescription);
+			const TAttributesSet<FVertexID> AttributesSet = MeshDescription.VertexAttributes();
+			TArray<FName> AttributeNames;
+			AttributesSet.ForEach([&AttributeNames, &AttributesSet](const FName AttributeName, const auto AttributesRef)
+			{
+				const bool bIsAutoGenerated = (AttributesRef.GetFlags() & EMeshAttributeFlags::AutoGenerated) != EMeshAttributeFlags::None;
+				if (!bIsAutoGenerated && AttributesSet.template HasAttributeOfType<float>(AttributeName))
+				{
+					AttributeNames.Add(AttributeName);
+				}
+			});
+			return AttributeNames;
+		}
 #endif
 	};
 }
@@ -312,7 +332,7 @@ UGeometryCache* UNearestNeighborModelSection::GetMutableNeighborMeshes() const
 void UNearestNeighborModelSection::SetVertexMapString(const FString& InString)
 {
 	VertexMapString = InString;
-	UpdateVertexMap();
+	UpdateVertexWeights();
 	check(Model);
 	Model->InvalidateTrainingModelOnly();
 }
@@ -346,6 +366,11 @@ const TArray<int32>&  UNearestNeighborModelSection::GetVertexMap() const
 	return VertexMap;
 }
 
+const TArray<float>& UNearestNeighborModelSection::GetVertexWeights() const
+{
+	return VertexWeights;
+}
+
 const TArray<float>& UNearestNeighborModelSection::GetPCABasis() const
 {
 	return PCABasis;
@@ -481,7 +506,8 @@ UE::NearestNeighborModel::EOpFlag UNearestNeighborModelSection::UpdateForTrainin
 		return EOpFlag::Error;
 	}
 
-	const EOpFlag Result = UpdateVertexMap();
+	EOpFlag Result = UpdateVertexWeights();
+	Result |= NormalizeVertexWeights();
 	using namespace UE::NearestNeighborModel;
 	if (!OpFlag::HasError(Result))
 	{
@@ -561,6 +587,38 @@ const UNearestNeighborModel* UNearestNeighborModelSection::GetModel() const
 	return Model;
 }
 
+ENearestNeighborModelSectionWeightMapCreationMethod UNearestNeighborModelSection::GetWeightMapCreationMethod() const
+{
+	return WeightMapCreationMethod;
+}
+
+FString UNearestNeighborModelSection::GetBoneNamesString() const
+{
+	FString Result;
+	if (BoneNames.IsEmpty())
+	{
+		return TEXT("None");
+	}
+	for (int32 Index = 0; Index < BoneNames.Num() - 1; ++Index)
+	{
+		Result += BoneNames[Index].ToString() + TEXT(", ");
+	}
+	Result += BoneNames.Last().ToString();
+	return Result;
+}
+
+const TArray<FName>& UNearestNeighborModelSection::GetBoneNames() const
+{
+	return BoneNames;
+}
+
+void UNearestNeighborModelSection::SetBoneNames(const TArray<FName>& InBoneNames)
+{
+	BoneNames = InBoneNames;
+	InvalidateTraining();
+	UpdateVertexWeights();
+}
+
 void UNearestNeighborModelSection::PostEditChangeProperty(FPropertyChangedEvent& PropertyChangedEvent)
 {
 	const FProperty* Property = PropertyChangedEvent.Property;
@@ -577,10 +635,12 @@ void UNearestNeighborModelSection::PostEditChangeProperty(FPropertyChangedEvent&
 		Model->UpdateNetworkOutputDim();
 		InvalidateTraining();
 	}
-	if (Property->GetFName() == GET_MEMBER_NAME_CHECKED(UNearestNeighborModelSection, VertexMapString))
+	if (Property->GetFName() == GET_MEMBER_NAME_CHECKED(UNearestNeighborModelSection, WeightMapCreationMethod) ||
+		Property->GetFName() == GET_MEMBER_NAME_CHECKED(UNearestNeighborModelSection, VertexMapString) ||
+		Property->GetFName() == GET_MEMBER_NAME_CHECKED(UNearestNeighborModelSection, AttributeName))
 	{
 		InvalidateTraining();
-		UpdateVertexMap();
+		UpdateVertexWeights();
 	}
 	if (Property->GetFName() == GET_MEMBER_NAME_CHECKED(UNearestNeighborModelSection, NeighborPoses) ||
 		Property->GetFName() == GET_MEMBER_NAME_CHECKED(UNearestNeighborModelSection, NeighborMeshes) ||
@@ -652,7 +712,40 @@ void UNearestNeighborModelSection::ResetNearestNeighborData()
 	RuntimeNeighborCoeffs.Reset();
 }
 
-UE::NearestNeighborModel::EOpFlag UNearestNeighborModelSection::UpdateVertexMap()
+UE::NearestNeighborModel::EOpFlag UNearestNeighborModelSection::UpdateVertexWeights()
+{
+	EOpFlag Result = EOpFlag::Error;
+	if (WeightMapCreationMethod == ENearestNeighborModelSectionWeightMapCreationMethod::FromText)
+	{
+		Result = UpdateVertexWeightsFromText();
+	}
+	else if (WeightMapCreationMethod == ENearestNeighborModelSectionWeightMapCreationMethod::SelectedBones)
+	{
+		Result = UpdateVertexWeightsSelectedBones();
+	}
+	else if (WeightMapCreationMethod == ENearestNeighborModelSectionWeightMapCreationMethod::VertexAttributes)
+	{
+		Result = UpdateVertexWeightsVertexAttributes();
+	}
+	using namespace UE::NearestNeighborModel;
+	if (OpFlag::HasError(Result))
+	{
+		VertexMap.Reset();
+		VertexWeights.Reset();
+		return Result;
+	}
+
+	NumVertices = VertexMap.Num();
+	check(VertexWeights.Num() == NumVertices);
+	if (NumVertices <= NumPCACoeffs)
+	{
+		UE_LOG(LogNearestNeighborModel, Error, TEXT("NumVertices %d needs to be larger than NumPCACoeffs %d."), NumVertices, NumPCACoeffs);
+		Result |= EOpFlag::Error;
+	}
+	return Result;
+}
+
+UE::NearestNeighborModel::EOpFlag UNearestNeighborModelSection::UpdateVertexWeightsFromText()
 {
 	using namespace UE::NearestNeighborModel;
 	if (!GetModel())
@@ -688,8 +781,10 @@ UE::NearestNeighborModel::EOpFlag UNearestNeighborModelSection::UpdateVertexMap(
 		if (const int32 MaxIndex = FMath::Max(VertexMap); MaxIndex >= SkeletalMesh->GetNumImportedVertices())
 		{
 			UE_LOG(LogNearestNeighborModel, Error, TEXT("MaxIndex %d is invalid, SkeletalMesh has %d vertices."), MaxIndex, SkeletalMesh->GetNumImportedVertices());
+			VertexMap.Reset();
 			return EOpFlag::Error;
 		}
+		VertexWeights.Init(1, VertexMap.Num());
 	}
 	else
 	{
@@ -697,15 +792,194 @@ UE::NearestNeighborModel::EOpFlag UNearestNeighborModelSection::UpdateVertexMap(
 		return EOpFlag::Error;
 	}
 
-	NumVertices = VertexMap.Num();
-	if (NumVertices <= NumPCACoeffs)
+	return EOpFlag::Success;
+}
+
+UE::NearestNeighborModel::EOpFlag UNearestNeighborModelSection::UpdateVertexWeightsSelectedBones()
+{
+	using namespace UE::NearestNeighborModel;
+	if (!GetModel())
 	{
-		UE_LOG(LogNearestNeighborModel, Error, TEXT("NumVertices %d needs to be larger than NumPCACoeffs %d."), NumVertices, NumPCACoeffs);
+		UE_LOG(LogNearestNeighborModel, Error, TEXT("NearestNeighborModel is invalid."));
+		return EOpFlag::Error;
+	}
+	const USkeletalMesh* const SkeletalMesh = GetModel()->GetSkeletalMesh();
+	if (!SkeletalMesh)
+	{
+		UE_LOG(LogNearestNeighborModel, Error, TEXT("SkeletalMesh is None"));
+		return EOpFlag::Error;
+	}
+	if (!SkeletalMesh->GetImportedModel())
+	{
+		UE_LOG(LogNearestNeighborModel, Error, TEXT("SkeletalMesh has no imported model."));
 		return EOpFlag::Error;
 	}
+	constexpr int32 LODIndex = 0;
+	if (!SkeletalMesh->GetImportedModel()->LODModels.IsValidIndex(LODIndex))
+	{
+		UE_LOG(LogNearestNeighborModel, Error, TEXT("SkeletalMesh has no LODModel."));
+		return EOpFlag::Error;
+	}
+	const FSkeletalMeshLODModel& LODModel = SkeletalMesh->GetImportedModel()->LODModels[LODIndex];
+	TArray<int32> BoneIndices;
+	BoneIndices.Reserve(BoneNames.Num());
+	for (const FName& BoneName : BoneNames)
+	{
+		const int32 BoneIndex = SkeletalMesh->GetRefSkeleton().FindBoneIndex(BoneName);
+		if (BoneIndex == INDEX_NONE)
+		{
+			UE_LOG(LogNearestNeighborModel, Error, TEXT("Bone %s does not exist in SkeletalMesh."), *BoneName.ToString());
+			return EOpFlag::Error;
+		}
+		BoneIndices.Add(BoneIndex);
+	}
+	const int32 NumImportedVertices = SkeletalMesh->GetNumImportedVertices();
+	TArray<float> Weights; 
+	Weights.SetNumZeroed(NumImportedVertices);
+	const TArray<FSkelMeshSection>& Sections = LODModel.Sections;
+	for (const FSkelMeshSection& Section : Sections)
+	{
+		if (Section.SoftVertices.Num() != Section.NumVertices)
+		{
+			UE_LOG(LogNearestNeighborModel, Error, TEXT("Section %d has invalid soft vertices."), Section.MaterialIndex);
+			return EOpFlag::Error;
+		}
+		for (int32 VertexIndex = 0; VertexIndex < Section.NumVertices; ++VertexIndex)
+		{
+			const FSoftSkinVertex& SoftVertex = Section.SoftVertices[VertexIndex];
+			const int32 ImportedIndex = LODModel.MeshToImportVertexMap[Section.BaseVertexIndex + VertexIndex];
+			if (Weights[ImportedIndex] > 0) // already processed
+			{
+				continue;
+			}
+			for (int32 InfluenceIndex = 0; InfluenceIndex < MAX_TOTAL_INFLUENCES; ++InfluenceIndex)
+			{
+				const uint16 WeightByte = SoftVertex.InfluenceWeights[InfluenceIndex];
+				if (WeightByte == 0)
+				{
+					break;
+				}
+				const int32 BoneIndex = Section.BoneMap[SoftVertex.InfluenceBones[InfluenceIndex]];
+				if (BoneIndices.Contains(BoneIndex))
+				{
+					const float	Weight = static_cast<float>(WeightByte) * UE::AnimationCore::InvMaxRawBoneWeightFloat;
+					Weights[ImportedIndex] += Weight;
+				}
+			}
+		}
+	}
+	VertexMap.Reset();
+	VertexMap.Reserve(NumImportedVertices);
+	VertexWeights.Reset();
+	VertexWeights.Reserve(NumImportedVertices);
+	for (int32 VertexIndex = 0; VertexIndex < NumImportedVertices; ++VertexIndex)
+	{
+		if (Weights[VertexIndex] > 0)
+		{
+			VertexMap.Add(VertexIndex);
+			VertexWeights.Add(Weights[VertexIndex]);
+		}
+	}
+	return EOpFlag::Success;
+}
+
+UE::NearestNeighborModel::EOpFlag UNearestNeighborModelSection::UpdateVertexWeightsVertexAttributes()
+{
+	using namespace UE::NearestNeighborModel;
+	if (!GetModel())
+	{
+		UE_LOG(LogNearestNeighborModel, Error, TEXT("NearestNeighborModel is invalid."));
+		return EOpFlag::Error;
+	}
+	const USkeletalMesh* const SkeletalMesh = GetModel()->GetSkeletalMesh();
+	if (!SkeletalMesh)
+	{
+		UE_LOG(LogNearestNeighborModel, Error, TEXT("SkeletalMesh is None"));
+		return EOpFlag::Error;
+	}
+
+	if (!SkeletalMesh->GetImportedModel())
+	{
+		UE_LOG(LogNearestNeighborModel, Error, TEXT("SkeletalMesh has no imported model."));
+		return EOpFlag::Error;
+	}
+	constexpr int32 LODIndex = 0;
+	if (!SkeletalMesh->GetImportedModel()->LODModels.IsValidIndex(LODIndex))
+	{
+		UE_LOG(LogNearestNeighborModel, Error, TEXT("SkeletalMesh has no LODModel."));
+		return EOpFlag::Error;
+	}
+	const FSkeletalMeshLODModel& LODModel = SkeletalMesh->GetImportedModel()->LODModels[LODIndex];
+
+	FMeshDescription MeshDescription;
+	SkeletalMesh->GetMeshDescription(LODIndex, MeshDescription);
+	const TAttributesSet<FVertexID> AttributesSet = MeshDescription.VertexAttributes();
+	if (!AttributesSet.HasAttributeOfType<float>(AttributeName))
+	{
+		UE_LOG(LogNearestNeighborModel, Error, TEXT("There is no float attribute %s in SkeletalMesh."), *AttributeName.ToString());
+		return EOpFlag::Error;
+	}
+	TVertexAttributesConstRef<float> AttributesRef = AttributesSet.GetAttributesRef<float>(AttributeName);
+	
+	const int32 NumImportedVertices = SkeletalMesh->GetNumImportedVertices();
+	VertexMap.Reset();
+	VertexMap.Reserve(NumImportedVertices);
+	VertexWeights.Reset();
+	VertexWeights.Reserve(NumImportedVertices);
+	for (const FVertexID VertexID : MeshDescription.Vertices().GetElementIDs())
+	{
+		const float Value = AttributesRef[VertexID];
+		if (Value > 0)
+		{
+			const int32 Index = LODModel.MeshToImportVertexMap[VertexID.GetValue()];
+			VertexMap.Add(VertexID.GetValue());
+			VertexWeights.Add(Value);
+		}
+	}
+
+	return EOpFlag::Success;
+}
+
+UE::NearestNeighborModel::EOpFlag UNearestNeighborModelSection::NormalizeVertexWeights()
+{
+	if (!Model)
+	{
+		UE_LOG(LogNearestNeighborModel, Error, TEXT("NearestNeighborModel in Section is invalid"));
+		return EOpFlag::Error;
+	}
+	const TArray<float>& WeightSum = Model->GetVertexWeightSum();
+	if (WeightSum.Num() != Model->GetNumBaseMeshVerts())
+	{
+		UE_LOG(LogNearestNeighborModel, Error, TEXT("VertexWeightSum is not initialized. Call UNearestNeighborModel::UpdateForTraining() instead."));
+		return EOpFlag::Error;
+	}
+	for (int32 Index = 0; Index < VertexWeights.Num(); ++Index)
+	{
+		const int32 VertexIndex = VertexMap[Index];
+		const float Sum = WeightSum[VertexIndex];
+		if (Sum > 1) // Only normalize if Sum > 1
+		{
+			VertexWeights[Index] /= Sum;
+		}
+	}
 	return EOpFlag::Success;
 }
-FString UNearestNeighborModelSection::BackwardCompatibleVertexString = TEXT("Saved with backward version");
+
+TArray<FName> UNearestNeighborModelSection::GetVertexAttributeNames() const
+{
+	if (!Model)
+	{
+		return {};
+	}
+	const USkeletalMesh* const SkeletalMesh = Model->GetSkeletalMesh();
+	if (!SkeletalMesh)
+	{
+		return {};
+	}
+
+	using UE::NearestNeighborModel::Private::GetVertexFloatAttributeNames;
+	return GetVertexFloatAttributeNames(*SkeletalMesh);
+}
 #endif
 
 PRAGMA_DISABLE_DEPRECATION_WARNINGS 
@@ -1193,6 +1467,26 @@ void UNearestNeighborModel::UpdateMorphTargetsLastWriteTime()
 	MorphTargetsLastWriteTime = FDateTime::UtcNow().GetTicks();
 }
 
+const TArray<float>& UNearestNeighborModel::GetVertexWeightSum() const
+{
+	return VertexWeightSum;
+}
+
+void UNearestNeighborModel::ComputeVertexWeightSum()
+{
+	VertexWeightSum.Init(0.f, GetNumBaseMeshVerts());
+	for (const FSection* Section : Sections)
+	{
+		check(Section);
+		const TArray<int32>& SectionVertexMap = Section->GetVertexMap();
+		const TArray<float>& SectionVertexWeights = Section->GetVertexWeights();
+		for (int32 Index = 0; Index < SectionVertexMap.Num(); ++Index)
+		{
+			VertexWeightSum[SectionVertexMap[Index]] += SectionVertexWeights[Index];
+		}
+	}
+}
+
 UE::NearestNeighborModel::EOpFlag UNearestNeighborModel::UpdateForTraining()
 {
 	if (IsReadyForTraining())
@@ -1225,11 +1519,14 @@ UE::NearestNeighborModel::EOpFlag UNearestNeighborModel::UpdateForTraining()
 		UE_LOG(LogNearestNeighborModel, Error, TEXT("At least one section is required. Please create a section using the '+' button."));
 		return EOpFlag::Error;
 	}
+	ComputeVertexWeightSum();
 	for (FSection* Section : Sections)
 	{
 		Section->SetModel(this);
 		Result |= Section->UpdateForTraining();
 	}
+	VertexWeightSum.Reset();
+
 	if (!OpFlag::HasError(Result))
 	{
 		bIsReadyForTraining = true;
diff --git a/Engine/Plugins/Animation/MLDeformer/NearestNeighborModel/Source/NearestNeighborModel/Private/NearestNeighborModelInstance.cpp b/Engine/Plugins/Animation/MLDeformer/NearestNeighborModel/Source/NearestNeighborModel/Private/NearestNeighborModelInstance.cpp
--- a/Engine/Plugins/Animation/MLDeformer/NearestNeighborModel/Source/NearestNeighborModel/Private/NearestNeighborModelInstance.cpp
+++ b/Engine/Plugins/Animation/MLDeformer/NearestNeighborModel/Source/NearestNeighborModel/Private/NearestNeighborModelInstance.cpp
@@ -423,12 +423,12 @@ void UNearestNeighborModelInstance::RunNearestNeighborModel(float DeltaTime, flo
 			{
 				const float Sigma = NearestNeighborModel->GetRBFSigma();
 				using UE::NearestNeighborModel::Private::UpdateRBFWeights;
-				UpdateRBFWeights(SectionMorphWeights, SectionPreviousWeights, SquaredDistances, Sigma, OffsetWeight, DecayCoeff);
+				UpdateRBFWeights(SectionMorphWeights, SectionPreviousWeights, SquaredDistances, Sigma, OffsetWeight * ModelWeight, DecayCoeff);
 			}
 			else
 			{
 				using UE::NearestNeighborModel::Private::UpdateNearestNeighborWeights;
-				UpdateNearestNeighborWeights(SectionMorphWeights, SectionPreviousWeights, SquaredDistances, OffsetWeight, DecayCoeff);
+				UpdateNearestNeighborWeights(SectionMorphWeights, SectionPreviousWeights, SquaredDistances, OffsetWeight * ModelWeight, DecayCoeff);
 			}
 
 #if WITH_EDITOR
diff --git a/Engine/Plugins/Animation/MLDeformer/NearestNeighborModel/Source/NearestNeighborModel/Public/NearestNeighborModel.h b/Engine/Plugins/Animation/MLDeformer/NearestNeighborModel/Source/NearestNeighborModel/Public/NearestNeighborModel.h
--- a/Engine/Plugins/Animation/MLDeformer/NearestNeighborModel/Source/NearestNeighborModel/Public/NearestNeighborModel.h
+++ b/Engine/Plugins/Animation/MLDeformer/NearestNeighborModel/Source/NearestNeighborModel/Public/NearestNeighborModel.h
@@ -70,6 +70,19 @@ struct UE_DEPRECATED(5.4, "FClothPartData is deprecated. Use UNearestNeighborMod
 	TArray<float> NeighborCoeffs;
 };
 
+#if WITH_EDITORONLY_DATA
+UENUM()
+enum class ENearestNeighborModelSectionWeightMapCreationMethod : uint8
+{
+	/** Include all vertices from text with weight 1. */
+	FromText,
+	/** Use skinning weights from selected bones. */
+	SelectedBones,
+	/** Use weights from a vertex attribute. */
+	VertexAttributes,
+};
+#endif
+
 /**
  * The section of the nearest neighbor model.
  * Each section contains a set of vertices in the original skeletal mesh.
@@ -115,6 +128,7 @@ public:
 
 	UFUNCTION(BlueprintPure, Category = "Section")
 	const TArray<int32>& GetVertexMap() const;
+	const TArray<float>& GetVertexWeights() const;
 	
 	UFUNCTION(BlueprintPure, Category = "Section")
 	const TArray<float>& GetPCABasis() const;
@@ -131,6 +145,8 @@ public:
 
 	FMLDeformerGeomCacheTrainingInputAnim* GetInputAnim() const;
 
+	EOpFlag UpdateVertexWeights();
+
 	EOpFlag UpdateForTraining();
 	EOpFlag UpdateForInference();
 	void InvalidateTraining();
@@ -142,11 +158,18 @@ public:
 	void SetModel(UNearestNeighborModel* InModel);
 	const UNearestNeighborModel* GetModel() const;
 
+	ENearestNeighborModelSectionWeightMapCreationMethod GetWeightMapCreationMethod() const;
+	FString GetBoneNamesString() const;
+	const TArray<FName>& GetBoneNames() const;
+	void SetBoneNames(const TArray<FName>& InBoneNames);
+
 	static FName GetNumPCACoeffsPropertyName() { return GET_MEMBER_NAME_CHECKED(UNearestNeighborModelSection, NumPCACoeffs); }
 	static FName GetVertexMapStringPropertyName() { return GET_MEMBER_NAME_CHECKED(UNearestNeighborModelSection, VertexMapString); }
 	static FName GetNeighborPosesPropertyName() { return GET_MEMBER_NAME_CHECKED(UNearestNeighborModelSection, NeighborPoses); }
 	static FName GetNeighborMeshesPropertyName() { return GET_MEMBER_NAME_CHECKED(UNearestNeighborModelSection, NeighborMeshes); }
 	static FName GetExcludedFramesPropertyName() { return GET_MEMBER_NAME_CHECKED(UNearestNeighborModelSection, ExcludedFrames); }
+	static FName GetWeightMapCreationMethodPropertyName() { return GET_MEMBER_NAME_CHECKED(UNearestNeighborModelSection, WeightMapCreationMethod); }
+	static FName GetAttributeNamePropertyName() { return GET_MEMBER_NAME_CHECKED(UNearestNeighborModelSection, AttributeName); }
 #endif
 #if WITH_EDITORONLY_DATA
 PRAGMA_DISABLE_DEPRECATION_WARNINGS
@@ -157,7 +180,6 @@ PRAGMA_ENABLE_DEPRECATION_WARNINGS
 
 protected:
 #if WITH_EDITORONLY_DATA
-
 	/** Poses of the nearest neighbor ROM. */
 	UPROPERTY(EditAnywhere, Category = "Section")
 	TObjectPtr<UAnimSequence> NeighborPoses;
@@ -166,6 +188,18 @@ protected:
 	UPROPERTY(EditAnywhere, Category = "Section")
 	TObjectPtr<UGeometryCache> NeighborMeshes;
 
+	/** Method to create weight map for this section. */
+	UPROPERTY(EditAnywhere, Category = "Section")
+	ENearestNeighborModelSectionWeightMapCreationMethod WeightMapCreationMethod = ENearestNeighborModelSectionWeightMapCreationMethod::FromText;
+
+	/** Bone names used to create weight map. */
+	UPROPERTY()
+	TArray<FName> BoneNames;
+	
+	/** A float vertex attribute that is used to compute weight maps */
+	UPROPERTY(EditAnywhere, Category = "Section", meta = (GetOptions = "GetVertexAttributeNames", NoResetToDefault))
+	FName AttributeName;
+
 	/** A string containing vertex indices for this section, e.g. "2, 3, 5-8, 9, 11-20" */
 	UPROPERTY(EditAnywhere, Category = "Section", meta = (DisplayName = "Vertex Indices"))
 	FString VertexMapString;
@@ -182,10 +216,14 @@ protected:
 	UPROPERTY()
 	TObjectPtr<UNearestNeighborModel> Model = nullptr;
 
-	/** Vertex indices for this section */
+	/** Vertex indices for this section. VertexMap.Num() == NumVertices */
 	UPROPERTY()
 	TArray<int32> VertexMap;
-	
+
+	/** The vertex weights for this section. VertexWeights.Num() == NumVertices. */
+	UPROPERTY()
+	TArray<float> VertexWeights;
+
 	/** Flattened array of PCA basis. The shape of PCA basis is (PCACoeffNum, NumVertices * 3)  */
 	UPROPERTY()
 	TArray<float> PCABasis;
@@ -247,10 +285,14 @@ private:
 	void Reset();
 	void ResetPCAData();
 	void ResetNearestNeighborData();
-	EOpFlag UpdateVertexMap();
-	
-	static FString BackwardCompatibleVertexString;
+	EOpFlag UpdateVertexWeightsFromText();
+	EOpFlag UpdateVertexWeightsSelectedBones();
+	EOpFlag UpdateVertexWeightsVertexAttributes();
+	EOpFlag NormalizeVertexWeights();
 
+	UFUNCTION()
+	TArray<FName> GetVertexAttributeNames() const;
+	
 	/** A temporary InputAnim used by GetInputAnim() function. */
 	mutable TUniquePtr<FMLDeformerGeomCacheTrainingInputAnim> InputAnim;
 #endif
@@ -401,6 +443,7 @@ public:
 	void UpdateNetworkOutputDim();
 
 	bool IsBeforeCustomVersionWasAdded() const;
+	const TArray<float>& GetVertexWeightSum() const;
 
 	static FName GetInputDimPropertyName() { return GET_MEMBER_NAME_CHECKED(UNearestNeighborModel, InputDim); }
 	static FName GetHiddenLayerDimsPropertyName() { return GET_MEMBER_NAME_CHECKED(UNearestNeighborModel, HiddenLayerDims); }
@@ -524,6 +567,8 @@ private:
 	void UpdateCachedDeltasTimestamp();
 	void UpdateCachedPCATimestamp();
 	void UpdateCachedNetworkTimestamp();
+
+	void ComputeVertexWeightSum();
 	// FNearestNeighborModelDetails needs to call private function GetSection(int32).
 	friend class UE::NearestNeighborModel::FNearestNeighborModelDetails;
 #endif
@@ -554,5 +599,7 @@ private:
 
 	UPROPERTY()
 	FString NetworkLastWriteArchitectureString;
+
+	TArray<float> VertexWeightSum;
 #endif
 };
diff --git a/Engine/Plugins/Animation/MLDeformer/NearestNeighborModel/Source/NearestNeighborModelEditor/Private/NearestNeighborEditorModel.cpp b/Engine/Plugins/Animation/MLDeformer/NearestNeighborModel/Source/NearestNeighborModelEditor/Private/NearestNeighborEditorModel.cpp
--- a/Engine/Plugins/Animation/MLDeformer/NearestNeighborModel/Source/NearestNeighborModelEditor/Private/NearestNeighborEditorModel.cpp
+++ b/Engine/Plugins/Animation/MLDeformer/NearestNeighborModel/Source/NearestNeighborModelEditor/Private/NearestNeighborEditorModel.cpp
@@ -29,15 +29,17 @@ namespace UE::NearestNeighborModel
 {
 	namespace Private
 	{
-		void AddFloatArrayToDeltaArray(TConstArrayView<float> FloatArr, TConstArrayView<int32> VertexMap, TArrayView<FVector3f> OutDeltaArr)
+		void AddFloatArrayToDeltaArray(TConstArrayView<float> FloatArr, TConstArrayView<int32> VertexMap, TConstArrayView<float> VertexWeights, TArrayView<FVector3f> OutDeltaArr)
 		{
 			const int32 SectionNumVerts = VertexMap.Num();
 			check(FloatArr.Num() == SectionNumVerts * 3);
+			check(VertexWeights.Num() == SectionNumVerts);
 			for (int32 Index = 0; Index < SectionNumVerts; ++Index)
 			{
 				const int32 FloatIndex = Index * 3;
 				const int32 DeltaIndex = VertexMap[Index];
-				OutDeltaArr[DeltaIndex] = FVector3f(FloatArr[FloatIndex], FloatArr[FloatIndex + 1], FloatArr[FloatIndex + 2]);
+				const float Weight = VertexWeights[Index];
+				OutDeltaArr[DeltaIndex] = FVector3f(FloatArr[FloatIndex], FloatArr[FloatIndex + 1], FloatArr[FloatIndex + 2]) * Weight;
 			}
 		}
 	};
@@ -321,7 +323,8 @@ namespace UE::NearestNeighborModel
 			// Draw the verts for the current frame.
 			if (bDrawVerts)
 			{
-				const FLinearColor VertsColor = FNearestNeighborModelEditorStyle::Get().GetColor("NearestNeighborModel.Verts.VertsColor");
+				const FLinearColor VertsColor0 = FNearestNeighborModelEditorStyle::Get().GetColor("NearestNeighborModel.Verts.VertsColor0");
+				const FLinearColor VertsColor1 = FNearestNeighborModelEditorStyle::Get().GetColor("NearestNeighborModel.Verts.VertsColor1");
 				const uint8 DepthGroup = VizSettings->GetXRayDeltas() ? 100 : 0;
 				const float PointSize = 5;
 				const int32 SectionIndex = VizSettings->VertVizSectionIndex;
@@ -331,9 +334,12 @@ namespace UE::NearestNeighborModel
 					if (NearestNeighborModel)
 					{
 						TConstArrayView<int32> VertexMap = NearestNeighborModel->GetSection(SectionIndex).GetVertexMap();
+						TConstArrayView<float> VertexWeights = NearestNeighborModel->GetSection(SectionIndex).GetVertexWeights();
 						for (int32 Index = 0; Index < VertexMap.Num(); ++Index)
 						{
 							const int32 ArrayIndex = 3 * VertexMap[Index];
+							const float Weight = VertexWeights[Index];
+							const FLinearColor VertsColor = FMath::Lerp(VertsColor0, VertsColor1, Weight);
 							const FVector VertexPos = (FVector)LinearSkinnedPositions[VertexMap[Index]];
 							PDI->DrawPoint(VertexPos, VertsColor, PointSize, DepthGroup);
 						}
@@ -344,7 +350,7 @@ namespace UE::NearestNeighborModel
 					for (int32 Index = 0; Index < LinearSkinnedPositions.Num(); ++Index)
 					{
 						const FVector VertexPos = (FVector)LinearSkinnedPositions[Index];
-						PDI->DrawPoint(VertexPos, VertsColor, PointSize, DepthGroup);
+						PDI->DrawPoint(VertexPos, VertsColor1, PointSize, DepthGroup);
 					}
 				}
 			}
@@ -639,14 +645,17 @@ namespace UE::NearestNeighborModel
 		{
 			const FSection& Section = NearestNeighborModel->GetSection(SectionIndex);
 			TConstArrayView<int32> VertexMap = Section.GetVertexMap();
+			TConstArrayView<float> VertexWeights = Section.GetVertexWeights();
 			if (VertexMap.IsEmpty())
 			{
 				UE_LOG(LogNearestNeighborModel, Warning, TEXT("Section %d has empty vertex map. No morph targets are generated for this section"), SectionIndex);
 				Result |= EOpFlag::Warning;
 				continue;
 			}
+			check(VertexMap.Num() == Section.GetNumVertices());
+			check(VertexWeights.Num() == Section.GetNumVertices());
 			TArrayView<FVector3f> MeanDeltas(Deltas.GetData(), NumBaseMeshVerts);
-			AddFloatArrayToDeltaArray(Section.GetVertexMean(), VertexMap, MeanDeltas);
+			AddFloatArrayToDeltaArray(Section.GetVertexMean(), VertexMap, VertexWeights, MeanDeltas);
 			const int32 NumPCACoeffs = Section.GetNumPCACoeffs();
 			const int32 SectionNumVerts = Section.GetNumVertices();
 			check(VertexMap.Num() == SectionNumVerts);
@@ -654,7 +663,7 @@ namespace UE::NearestNeighborModel
 			{
 				TConstArrayView<float> BasisFloats(Section.GetPCABasis().GetData() + Index * SectionNumVerts * 3, SectionNumVerts * 3);
 				TArrayView<FVector3f> BasisDeltas(Deltas.GetData() + (MorphOffset + Index) * NumBaseMeshVerts, NumBaseMeshVerts);
-				AddFloatArrayToDeltaArray(BasisFloats, VertexMap, BasisDeltas);
+				AddFloatArrayToDeltaArray(BasisFloats, VertexMap, VertexWeights, BasisDeltas);
 			}
 			MorphOffset += NumPCACoeffs;
 		}
@@ -663,6 +672,7 @@ namespace UE::NearestNeighborModel
 		{
 			const FSection& Section = NearestNeighborModel->GetSection(SectionIndex);
 			TConstArrayView<int32> VertexMap = Section.GetVertexMap();
+			TConstArrayView<float> VertexWeights = Section.GetVertexWeights();
 			if (VertexMap.IsEmpty())
 			{
 				// Warning already generated
@@ -683,7 +693,7 @@ namespace UE::NearestNeighborModel
 			{
 				TConstArrayView<float> NeighborOffsets(RuntimeNeighborOffsets.GetData() + Index * SectionNumVerts * 3, SectionNumVerts * 3);
 				TArrayView<FVector3f> NeighborDeltas(Deltas.GetData() + (MorphOffset + Index) * NumBaseMeshVerts, NumBaseMeshVerts);
-				AddFloatArrayToDeltaArray(NeighborOffsets, VertexMap, NeighborDeltas);
+				AddFloatArrayToDeltaArray(NeighborOffsets, VertexMap, VertexWeights, NeighborDeltas);
 			}
 		}
 	
diff --git a/Engine/Plugins/Animation/MLDeformer/NearestNeighborModel/Source/NearestNeighborModelEditor/Private/NearestNeighborModelSectionCustomization.cpp b/Engine/Plugins/Animation/MLDeformer/NearestNeighborModel/Source/NearestNeighborModelEditor/Private/NearestNeighborModelSectionCustomization.cpp
--- a/Engine/Plugins/Animation/MLDeformer/NearestNeighborModel/Source/NearestNeighborModelEditor/Private/NearestNeighborModelSectionCustomization.cpp
+++ b/Engine/Plugins/Animation/MLDeformer/NearestNeighborModel/Source/NearestNeighborModelEditor/Private/NearestNeighborModelSectionCustomization.cpp
@@ -4,10 +4,18 @@
 
 #include "DetailWidgetRow.h"
 #include "DetailLayoutBuilder.h"
+#include "Engine/SkeletalMesh.h"
 #include "MLDeformerEditorModule.h"
+#include "MLDeformerEditorStyle.h"
 #include "NearestNeighborModel.h"
 #include "NearestNeighborEditorModel.h"
+#include "Rendering/SkeletalMeshLODImporterData.h" 
+#include "Rendering/SkeletalMeshModel.h"
+#include "SMLDeformerBonePickerDialog.h"
+#include "Widgets/Input/SButton.h"
+#include "Widgets/Input/SEditableTextBox.h"
 #include "Widgets/Input/STextComboBox.h"
+#include "Widgets/Text/STextBlock.h"
 
 #define LOCTEXT_NAMESPACE "NearestNeighborModelSectionCustomization"
 
@@ -25,6 +33,228 @@ namespace UE::NearestNeighborModel
 			FMLDeformerEditorModule& EditorModule = FModuleManager::LoadModuleChecked<FMLDeformerEditorModule>("MLDeformerFrameworkEditor");
 			return static_cast<FNearestNeighborEditorModel*>(EditorModule.GetModelRegistry().GetEditorModel(const_cast<UNearestNeighborModel*>(Model)));
 		}
+
+		TArray<float> GetVertexAttributeValues(const TArray<int32>& VertexMap, const TArray<float>& VertexWeights, int32 NumVertices)
+		{
+			check(VertexMap.Num() == VertexWeights.Num());
+			check(VertexMap.Num() <= NumVertices);
+			TArray<float> VertexAttribute;
+			VertexAttribute.SetNumZeroed(NumVertices);
+			for (int32 VertexIndex = 0; VertexIndex < VertexMap.Num(); ++VertexIndex)
+			{
+				VertexAttribute[VertexMap[VertexIndex]] = VertexWeights[VertexIndex];
+			}
+			return VertexAttribute;
+		}
+
+		void CreateVertexAttributes(USkeletalMesh& SkeletalMesh, const FString& AttributeName, const TArray<int32>& VertexMap, const TArray<float>& VertexWeights)
+		{
+			FSkeletalMeshImportData ImportData; 
+			constexpr int32 LODIndex = 0;
+			SkeletalMesh.LoadLODImportedData(LODIndex, ImportData);
+			using SkeletalMeshImportData::FVertexAttribute;
+			constexpr int32 NumComponents = 1;
+			FVertexAttribute VertexAttribute(GetVertexAttributeValues(VertexMap, VertexWeights, SkeletalMesh.GetNumImportedVertices()), NumComponents);
+			ImportData.VertexAttributeNames.Add(AttributeName);
+			ImportData.VertexAttributes.Add(VertexAttribute);
+			SkeletalMesh.SaveLODImportedData(LODIndex, ImportData);
+		}
+
+
+		class SNewAttributesDialog
+			: public SCustomDialog
+		{
+		public:
+			SLATE_BEGIN_ARGS(SNewAttributesDialog) {}
+			SLATE_END_ARGS()
+
+			void Construct(const FArguments& InArgs);
+			const FString& GetAttributesName() const { return AttributesName; }
+		private:
+			FString AttributesName;
+		};
+
+		void SNewAttributesDialog::Construct(const FArguments& InArgs)
+		{
+			FText DialogTitle = LOCTEXT("EnterAttributesNameDialogTitle", "New Vertex Attributes");
+			SCustomDialog::Construct
+			(
+				SCustomDialog::FArguments()
+				.AutoCloseOnButtonPress(true)
+				.Title(DialogTitle)
+				.UseScrollBox(false)
+				.Buttons(
+				{
+					SCustomDialog::FButton(LOCTEXT("OKText", "OK"))
+					.SetPrimary(true)
+					.SetFocus(),
+					SCustomDialog::FButton(LOCTEXT("CancelText", "Cancel"))
+				})
+				.Content()
+				[
+					SNew(SBox)
+					.Padding(FMargin(10.0f, 10.0f))
+					.MinDesiredWidth(400.0f)
+					.HAlign(EHorizontalAlignment::HAlign_Fill)
+					.VAlign(EVerticalAlignment::VAlign_Fill)
+					[
+						SNew(SHorizontalBox)
+						+ SHorizontalBox::Slot()
+						.AutoWidth()
+						.VAlign(EVerticalAlignment::VAlign_Center)
+						[
+							SNew(STextBlock)
+							.Text(LOCTEXT("AttributesName", "Attributes Name:"))
+						]
+						+ SHorizontalBox::Slot()
+						.VAlign(EVerticalAlignment::VAlign_Center)
+						[
+							SNew(SBox)
+							.MinDesiredWidth(200.0f)
+							.Padding(FMargin(4.0f, 0.0f))
+							[
+								SNew(SEditableTextBox)
+								.OnTextChanged_Lambda([&AttributesName = AttributesName](const FText& InText)
+								{
+									AttributesName = InText.ToString();
+								})
+							]
+						]
+					]
+				]
+			);
+		}
+
+		void CreateWeightMapWidgetSelectedBones(FDetailWidgetRow& Row, IDetailLayoutBuilder& LayoutBuilder, UNearestNeighborModelSection& Section, const FNearestNeighborEditorModel& EditorModel)
+		{
+			Row.NameContent()
+			[
+				SNew(STextBlock)
+				.Text(LOCTEXT("SelectedBones", "Selected Bones"))
+				.Font(IDetailLayoutBuilder::GetDetailFont())
+			]
+			.ValueContent()
+			[
+				SNew(SVerticalBox)
+				+ SVerticalBox::Slot()
+				.Padding(4, 2)
+				.AutoHeight()
+				[
+					SNew(STextBlock)
+					.Text_Lambda([&Section]()
+					{
+						return FText::FromString(Section.GetBoneNamesString());
+					})
+					.Font(IDetailLayoutBuilder::GetDetailFont())
+					.ToolTipText_Lambda([&Section]()
+					{
+						return FText::FromString(Section.GetBoneNamesString());
+					})
+				]
+				+ SVerticalBox::Slot()
+				.AutoHeight()
+				[
+					SNew(SHorizontalBox)
+					+ SHorizontalBox::Slot()
+					.AutoWidth()
+					[
+						SNew(SButton)
+						.Text(LOCTEXT("SelectBones", "Select Bones"))
+						.OnClicked_Lambda([&Section, &EditorModel]()
+						{
+							UNearestNeighborModel* const Model = Cast<UNearestNeighborModel>(EditorModel.GetModel());
+							if (!Model)
+							{
+								return FReply::Handled();
+							}
+							USkeletalMesh* const SkelMesh = Model->GetSkeletalMesh();
+							if (!SkelMesh)
+							{
+								return FReply::Handled();
+							}
+							const FLinearColor HighlightColor = UE::MLDeformer::FMLDeformerEditorStyle::Get().GetColor("MLDeformer.InputsWidget.HighlightColor");
+							TSharedPtr<UE::MLDeformer::SMLDeformerBonePickerDialog> Dialog = 
+								SNew(UE::MLDeformer::SMLDeformerBonePickerDialog)
+								.RefSkeleton(&SkelMesh->GetRefSkeleton())
+								.AllowMultiSelect(true)
+								.HighlightBoneNamesColor(FSlateColor(HighlightColor))
+								.HighlightBoneNames(Section.GetBoneNames());
+							
+							Dialog->ShowModal();
+							const TArray<FName>& BoneNames = Dialog->GetPickedBoneNames();
+							if (!BoneNames.IsEmpty())
+							{
+								Section.SetBoneNames(BoneNames);
+							}
+							return FReply::Handled();
+						})
+					]
+					+ SHorizontalBox::Slot()
+					[
+						SNew(SButton)
+						.Text(LOCTEXT("CreateAttributes", "Create Attributes"))
+						.OnClicked_Lambda([&Section, &EditorModel]()
+						{
+							TSharedPtr<SNewAttributesDialog> Dialog = SNew(SNewAttributesDialog);
+							Dialog->ShowModal();
+							const FString& AttributesName = Dialog->GetAttributesName();
+							UE_LOG(LogNearestNeighborModel, Log, TEXT("Create Attributes: %s"), *AttributesName);
+
+							UNearestNeighborModel* const Model = Cast<UNearestNeighborModel>(EditorModel.GetModel());
+							if (!Model || !Model->GetSkeletalMesh())
+							{
+								return FReply::Handled();
+							}
+
+							CreateVertexAttributes(*Model->GetSkeletalMesh(), AttributesName, Section.GetVertexMap(), Section.GetVertexWeights());
+							return FReply::Handled();
+						})
+					]
+				]
+			];
+		}
+
+		void CreateWeightMapWidgetFromText(FDetailWidgetRow& Row, IDetailLayoutBuilder& LayoutBuilder,UNearestNeighborModelSection& Section, const FNearestNeighborEditorModel& EditorModel)
+		{
+			TSharedPtr<IPropertyHandle> VertexMapHandle = LayoutBuilder.GetProperty(UNearestNeighborModelSection::GetVertexMapStringPropertyName());
+			if (VertexMapHandle.IsValid())
+			{
+				Row.NameContent()
+				[
+					VertexMapHandle->CreatePropertyNameWidget()
+				]
+				.ValueContent()
+				[
+					SNew(SVerticalBox)
+					+ SVerticalBox::Slot()
+					.Padding(2, 2)
+					.AutoHeight()
+					[
+						SNew(STextComboBox)
+						.OptionsSource(EditorModel.GetVertexMapSelector()->GetOptions())
+						.OnSelectionChanged_Lambda([&Section, &EditorModel](TSharedPtr<FString> InString, ESelectInfo::Type InSelectInfo)
+							{
+								EditorModel.GetVertexMapSelector()->OnSelectionChanged(Section, InString, InSelectInfo);
+							})
+						.InitiallySelectedItem(EditorModel.GetVertexMapSelector()->GetSelectedItem(Section))
+						.Font(IDetailLayoutBuilder::GetDetailFont())
+					]
+					+ SVerticalBox::Slot()
+					.AutoHeight()
+					[
+						SNew(SBox)
+						.IsEnabled(TAttribute<bool>::Create(TAttribute<bool>::FGetter::CreateLambda([&Section]()
+						{
+							return Section.GetMeshIndex() == INDEX_NONE;
+						})))
+						[
+							VertexMapHandle->CreatePropertyValueWidget()
+						]
+					]
+				];
+				LayoutBuilder.HideProperty(VertexMapHandle);
+			}
+		}
 	};
 
 	void SNearestNeighborModelSectionWidget::Construct(const FArguments& InArgs)
@@ -81,45 +311,36 @@ namespace UE::NearestNeighborModel
 			return;
 		}
 
-		TSharedPtr<IPropertyHandle> VertexMapHandle = DetailBuilder.GetProperty(UNearestNeighborModelSection::GetVertexMapStringPropertyName());
-		if (VertexMapHandle.IsValid())
-		{
-			SectionBuilder.AddProperty(VertexMapHandle).CustomWidget()
-				.OverrideResetToDefault(FResetToDefaultOverride::Hide())
-				.NameContent()
-				[
-					VertexMapHandle->CreatePropertyNameWidget()
-				]
-				.ValueContent()
-				[
-					SNew(SVerticalBox)
-					+ SVerticalBox::Slot()
-					.Padding(2, 2)
-					.AutoHeight()
-					[
-						SNew(STextComboBox)
-						.OptionsSource(EditorModel->GetVertexMapSelector()->GetOptions())
-						.OnSelectionChanged_Lambda([Section, EditorModel](TSharedPtr<FString> InString, ESelectInfo::Type InSelectInfo)
-							{
-								EditorModel->GetVertexMapSelector()->OnSelectionChanged(*Section, InString, InSelectInfo);
-							})
-						.InitiallySelectedItem(EditorModel->GetVertexMapSelector()->GetSelectedItem(*Section))
-						.Font(IDetailLayoutBuilder::GetDetailFont())
-					]
-					+ SVerticalBox::Slot()
-					.AutoHeight()
-					[
-						SNew(SBox)
-						.IsEnabled(TAttribute<bool>::Create(TAttribute<bool>::FGetter::CreateLambda([Section]()
-						{
-							return Section->GetMeshIndex() == INDEX_NONE;
-						})))
-						[
-							VertexMapHandle->CreatePropertyValueWidget()
-						]
-					]
-				];
-		}
+		SectionBuilder.AddProperty(UNearestNeighborModelSection::GetWeightMapCreationMethodPropertyName());
+
+		FDetailWidgetRow& WeightMapFromTextRow = SectionBuilder.AddCustomRow(LOCTEXT("WeightMapFromText", "WeightMapFromText"))
+			.OverrideResetToDefault(FResetToDefaultOverride::Hide())
+			.Visibility(TAttribute<EVisibility>::Create(TAttribute<EVisibility>::FGetter::CreateLambda([Section]()
+			{
+				return Section->GetWeightMapCreationMethod() == 
+				ENearestNeighborModelSectionWeightMapCreationMethod::FromText 
+				? EVisibility::Visible : EVisibility::Collapsed;
+			})));
+		Private::CreateWeightMapWidgetFromText(WeightMapFromTextRow, DetailBuilder, *Section, *EditorModel);
+
+		FDetailWidgetRow& WeightMapSelectedBonesRow = SectionBuilder.AddCustomRow(LOCTEXT("WeightMapSelectedBones", "WeightMapSelectedBones"))
+			.OverrideResetToDefault(FResetToDefaultOverride::Hide())
+			.Visibility(TAttribute<EVisibility>::Create(TAttribute<EVisibility>::FGetter::CreateLambda([Section]()
+			{
+				return Section->GetWeightMapCreationMethod() == 
+				ENearestNeighborModelSectionWeightMapCreationMethod::SelectedBones 
+				? EVisibility::Visible : EVisibility::Collapsed;
+			})));
+		Private::CreateWeightMapWidgetSelectedBones(WeightMapSelectedBonesRow, DetailBuilder, *Section, *EditorModel);
+
+		SectionBuilder.AddProperty(UNearestNeighborModelSection::GetAttributeNamePropertyName())
+			.OverrideResetToDefault(FResetToDefaultOverride::Hide())
+			.Visibility(TAttribute<EVisibility>::Create(TAttribute<EVisibility>::FGetter::CreateLambda([Section]()
+			{
+				return Section->GetWeightMapCreationMethod() == 
+				ENearestNeighborModelSectionWeightMapCreationMethod::VertexAttributes 
+				? EVisibility::Visible : EVisibility::Collapsed;
+			})));
 
 		const UNearestNeighborModel* const Model = Section->GetModel();
 		if (!Model)
diff --git a/Engine/Plugins/Animation/MLDeformer/NearestNeighborModel/Source/NearestNeighborModelEditor/Private/NearestNeighborModelStyle.cpp b/Engine/Plugins/Animation/MLDeformer/NearestNeighborModel/Source/NearestNeighborModelEditor/Private/NearestNeighborModelStyle.cpp
--- a/Engine/Plugins/Animation/MLDeformer/NearestNeighborModel/Source/NearestNeighborModelEditor/Private/NearestNeighborModelStyle.cpp
+++ b/Engine/Plugins/Animation/MLDeformer/NearestNeighborModel/Source/NearestNeighborModelEditor/Private/NearestNeighborModelStyle.cpp
@@ -19,7 +19,8 @@ namespace UE::NearestNeighborModel
 		// Colors and sizes.
 		Set("NearestNeighborModel.NearestNeighborActors.WireframeColor", FLinearColor(1.0f, 1.0f, 0.0f));
 		Set("NearestNeighborModel.NearestNeighborActors.LabelColor", FLinearColor(1.0f, 1.0f, 0.0f));
-		Set("NearestNeighborModel.Verts.VertsColor", FLinearColor(1.0f, 0.5f, 0.0f));
+		Set("NearestNeighborModel.Verts.VertsColor0", FLinearColor(0.7f, 0.7f, 0.7f));
+		Set("NearestNeighborModel.Verts.VertsColor1", FLinearColor(1.0f, 0.5f, 0.0f));
 
 		FSlateStyleRegistry::RegisterSlateStyle(*this);
 	}
