NearestNeighborModel: allow geomcache to change without changing model state

diff --git a/Engine/Plugins/Animation/MLDeformer/NearestNeighborModel/Source/NearestNeighborModel/Private/NearestNeighborModel.cpp b/Engine/Plugins/Animation/MLDeformer/NearestNeighborModel/Source/NearestNeighborModel/Private/NearestNeighborModel.cpp
--- a/Engine/Plugins/Animation/MLDeformer/NearestNeighborModel/Source/NearestNeighborModel/Private/NearestNeighborModel.cpp
+++ b/Engine/Plugins/Animation/MLDeformer/NearestNeighborModel/Source/NearestNeighborModel/Private/NearestNeighborModel.cpp
@@ -642,8 +642,9 @@ void UNearestNeighborModelSection::PostEditChangeProperty(FPropertyChangedEvent&
 		InvalidateTraining();
 		UpdateVertexWeights();
 	}
+	// Sometimes we need to set geometry cache to None before checking in. Do not invalidate training in this case.
 	if (Property->GetFName() == GET_MEMBER_NAME_CHECKED(UNearestNeighborModelSection, NeighborPoses) ||
-		Property->GetFName() == GET_MEMBER_NAME_CHECKED(UNearestNeighborModelSection, NeighborMeshes) ||
+		// Property->GetFName() == GET_MEMBER_NAME_CHECKED(UNearestNeighborModelSection, NeighborMeshes) ||
 		Property->GetFName() == GET_MEMBER_NAME_CHECKED(UNearestNeighborModelSection, ExcludedFrames))
 	{
 		InvalidateInference();
@@ -1041,7 +1042,6 @@ void UNearestNeighborModel::PostEditChangeProperty(FPropertyChangedEvent& Proper
 		return;
 	}
 	if (Property->GetFName() == UMLDeformerModel::GetSkeletalMeshPropertyName() ||
-		PropertyChangedEvent.GetMemberPropertyName() == TEXT("TrainingInputAnims") ||
 		Property->GetFName() == UMLDeformerModel::GetAlignmentTransformPropertyName() ||
 		Property->GetFName() == UMLDeformerModel::GetBoneIncludeListPropertyName() ||
 		Property->GetFName() == UMLDeformerModel::GetCurveIncludeListPropertyName() ||
@@ -1055,6 +1055,14 @@ void UNearestNeighborModel::PostEditChangeProperty(FPropertyChangedEvent& Proper
 	{
 		InvalidateTraining();
 	}
+	if (PropertyChangedEvent.GetMemberPropertyName() == TEXT("TrainingInputAnims"))
+	{
+		// Sometimes we need to set geometry cache to None before checking in. Do not invalidate training in this case.
+		if (Property->GetFName() != FMLDeformerGeomCacheTrainingInputAnim::GetGeomCachePropertyName())
+		{
+			InvalidateTraining();
+		}
+	}
 
 	if (Property->GetFName() == UNearestNeighborModel::GetBoneIncludeListPropertyName() ||
 		Property->GetFName() == UNearestNeighborModel::GetCurveIncludeListPropertyName())
