NearestNeighborModel: fix cook error when loading old asset

diff --git a/Engine/Plugins/Animation/MLDeformer/NearestNeighborModel/Source/NearestNeighborModel/Private/NearestNeighborModel.cpp b/Engine/Plugins/Animation/MLDeformer/NearestNeighborModel/Source/NearestNeighborModel/Private/NearestNeighborModel.cpp
--- a/Engine/Plugins/Animation/MLDeformer/NearestNeighborModel/Source/NearestNeighborModel/Private/NearestNeighborModel.cpp
+++ b/Engine/Plugins/Animation/MLDeformer/NearestNeighborModel/Source/NearestNeighborModel/Private/NearestNeighborModel.cpp
@@ -989,6 +989,7 @@ void UNearestNeighborModelSection::InitFromClothPartData(FClothPartData& InPart)
 {
 	using UE::NearestNeighborModel::Private::ConvertArray;
 	VertexMap = ConvertArray<uint32, int32>(MoveTemp(InPart.VertexMap));
+	VertexWeights.Init(1, VertexMap.Num());
 	using UE::NearestNeighborModel::Private::IntegersToFormattedString;
 	VertexMapString = IntegersToFormattedString(VertexMap);
 	MeshIndex = INDEX_NONE;
