NearestNeighborModel: fix normalize vertex weights

diff --git a/Engine/Plugins/Animation/MLDeformer/NearestNeighborModel/Source/NearestNeighborModel/Private/NearestNeighborModel.cpp b/Engine/Plugins/Animation/MLDeformer/NearestNeighborModel/Source/NearestNeighborModel/Private/NearestNeighborModel.cpp
--- a/Engine/Plugins/Animation/MLDeformer/NearestNeighborModel/Source/NearestNeighborModel/Private/NearestNeighborModel.cpp
+++ b/Engine/Plugins/Animation/MLDeformer/NearestNeighborModel/Source/NearestNeighborModel/Private/NearestNeighborModel.cpp
@@ -507,7 +507,6 @@ UE::NearestNeighborModel::EOpFlag UNearestNeighborModelSection::UpdateForTrainin
 	}
 
 	EOpFlag Result = UpdateVertexWeights();
-	Result |= NormalizeVertexWeights();
 	using namespace UE::NearestNeighborModel;
 	if (!OpFlag::HasError(Result))
 	{
@@ -1481,7 +1480,7 @@ const TArray<float>& UNearestNeighborModel::GetVertexWeightSum() const
 	return VertexWeightSum;
 }
 
-void UNearestNeighborModel::ComputeVertexWeightSum()
+void UNearestNeighborModel::NormalizeVertexWeights()
 {
 	VertexWeightSum.Init(0.f, GetNumBaseMeshVerts());
 	for (const FSection* Section : Sections)
@@ -1494,6 +1493,12 @@ void UNearestNeighborModel::ComputeVertexWeightSum()
 			VertexWeightSum[SectionVertexMap[Index]] += SectionVertexWeights[Index];
 		}
 	}
+	for (FSection* Section : Sections)
+	{
+		check(Section);
+		Section->NormalizeVertexWeights();
+	}
+	VertexWeightSum.Reset();
 }
 
 UE::NearestNeighborModel::EOpFlag UNearestNeighborModel::UpdateForTraining()
@@ -1528,12 +1533,12 @@ UE::NearestNeighborModel::EOpFlag UNearestNeighborModel::UpdateForTraining()
 		UE_LOG(LogNearestNeighborModel, Error, TEXT("At least one section is required. Please create a section using the '+' button."));
 		return EOpFlag::Error;
 	}
-	ComputeVertexWeightSum();
 	for (FSection* Section : Sections)
 	{
 		Section->SetModel(this);
 		Result |= Section->UpdateForTraining();
 	}
+	NormalizeVertexWeights();
 	VertexWeightSum.Reset();
 
 	if (!OpFlag::HasError(Result))
diff --git a/Engine/Plugins/Animation/MLDeformer/NearestNeighborModel/Source/NearestNeighborModel/Public/NearestNeighborModel.h b/Engine/Plugins/Animation/MLDeformer/NearestNeighborModel/Source/NearestNeighborModel/Public/NearestNeighborModel.h
--- a/Engine/Plugins/Animation/MLDeformer/NearestNeighborModel/Source/NearestNeighborModel/Public/NearestNeighborModel.h
+++ b/Engine/Plugins/Animation/MLDeformer/NearestNeighborModel/Source/NearestNeighborModel/Public/NearestNeighborModel.h
@@ -163,6 +163,9 @@ public:
 	const TArray<FName>& GetBoneNames() const;
 	void SetBoneNames(const TArray<FName>& InBoneNames);
 
+	// Do not call this function directly. Call UNearestNeighborModel::NormalizeVertexWeights() instead.
+	EOpFlag NormalizeVertexWeights();
+
 	static FName GetNumPCACoeffsPropertyName() { return GET_MEMBER_NAME_CHECKED(UNearestNeighborModelSection, NumPCACoeffs); }
 	static FName GetVertexMapStringPropertyName() { return GET_MEMBER_NAME_CHECKED(UNearestNeighborModelSection, VertexMapString); }
 	static FName GetNeighborPosesPropertyName() { return GET_MEMBER_NAME_CHECKED(UNearestNeighborModelSection, NeighborPoses); }
@@ -288,7 +291,6 @@ private:
 	EOpFlag UpdateVertexWeightsFromText();
 	EOpFlag UpdateVertexWeightsSelectedBones();
 	EOpFlag UpdateVertexWeightsVertexAttributes();
-	EOpFlag NormalizeVertexWeights();
 
 	UFUNCTION()
 	TArray<FName> GetVertexAttributeNames() const;
@@ -568,7 +570,7 @@ private:
 	void UpdateCachedPCATimestamp();
 	void UpdateCachedNetworkTimestamp();
 
-	void ComputeVertexWeightSum();
+	void NormalizeVertexWeights();
 	// FNearestNeighborModelDetails needs to call private function GetSection(int32).
 	friend class UE::NearestNeighborModel::FNearestNeighborModelDetails;
 #endif
