[MLDeformer] Fixed some crash blocker bug when creating a new ML Deformer asset and switching to Nearest neighbor model.

#trivial

diff --git a/Engine/Plugins/Animation/MLDeformer/MLDeformerFramework/Source/MLDeformerFramework/Public/MLDeformerMorphModel.h b/Engine/Plugins/Animation/MLDeformer/MLDeformerFramework/Source/MLDeformerFramework/Public/MLDeformerMorphModel.h
--- a/Engine/Plugins/Animation/MLDeformer/MLDeformerFramework/Source/MLDeformerFramework/Public/MLDeformerMorphModel.h
+++ b/Engine/Plugins/Animation/MLDeformer/MLDeformerFramework/Source/MLDeformerFramework/Public/MLDeformerMorphModel.h
@@ -215,7 +215,7 @@ public:
 	 */
 	UE_DEPRECATED(5.4, "This method has been deprecated, please use the GetMorphTargetSet that takes a LOD level parameter.")
 	TSharedPtr<FExternalMorphSet> GetMorphTargetSet() const				{ return MorphTargetSet_DEPRECATED; }
-	TSharedPtr<FExternalMorphSet> GetMorphTargetSet(int32 LOD) const	{ return MorphTargetSets.IsValidIndex(LOD) ? MorphTargetSets[LOD] : MorphTargetSets.Last(); }
+	TSharedPtr<FExternalMorphSet> GetMorphTargetSet(int32 LOD) const	{ return MorphTargetSets.IsValidIndex(LOD) ? MorphTargetSets[LOD] : TSharedPtr<FExternalMorphSet>(); }
 
 	/** Get the number of LOD levels that we have deltas for. */
 	int32 GetNumLODs() const											{ return MorphTargetSets.Num(); }
diff --git a/Engine/Plugins/Animation/MLDeformer/NearestNeighborModel/Source/NearestNeighborModelEditor/Private/NearestNeighborModelDetails.cpp b/Engine/Plugins/Animation/MLDeformer/NearestNeighborModel/Source/NearestNeighborModelEditor/Private/NearestNeighborModelDetails.cpp
--- a/Engine/Plugins/Animation/MLDeformer/NearestNeighborModel/Source/NearestNeighborModelEditor/Private/NearestNeighborModelDetails.cpp
+++ b/Engine/Plugins/Animation/MLDeformer/NearestNeighborModel/Source/NearestNeighborModelEditor/Private/NearestNeighborModelDetails.cpp
@@ -284,20 +284,23 @@ namespace UE::NearestNeighborModel
 				const int32 LOD = 0;
 				if (const UNearestNeighborModel* const NearestNeighborModel = GetCastModel())
 				{
-					const TSharedPtr<const FExternalMorphSet> MorphSet = NearestNeighborModel->GetMorphTargetSet(LOD);
-					if (MorphSet.IsValid() && MorphSet->MorphBuffers.IsMorphResourcesInitialized())
+					if (NearestNeighborModel->GetNumLODs() > 0)
 					{
-						if (NearestNeighborModel->IsBeforeCustomVersionWasAdded())
+						const TSharedPtr<const FExternalMorphSet> MorphSet = NearestNeighborModel->GetMorphTargetSet(LOD);
+						if (MorphSet.IsValid() && MorphSet->MorphBuffers.IsMorphResourcesInitialized())
 						{
-							return FText::FromString(TEXT("Loaded from previous version."));
-						}
-						else
-						{
-							const FDateTime Time = NearestNeighborModel->GetMorphTargetsLastWriteTime();
-							if (Time != FDateTime::MinValue())
+							if (NearestNeighborModel->IsBeforeCustomVersionWasAdded())
+							{
+								return FText::FromString(TEXT("Loaded from previous version."));
+							}
+							else
 							{
-								const FDateTime LocalTime = Private::ToLocalTime(Time);
-								return FText::FromString(FString::Printf(TEXT("Last written: %s"), *Private::ToString(LocalTime)));
+								const FDateTime Time = NearestNeighborModel->GetMorphTargetsLastWriteTime();
+								if (Time != FDateTime::MinValue())
+								{
+									const FDateTime LocalTime = Private::ToLocalTime(Time);
+									return FText::FromString(FString::Printf(TEXT("Last written: %s"), *Private::ToString(LocalTime)));
+								}
 							}
 						}
 					}
