NearestNeighborModel: remove onnx dependency

diff --git a/Engine/Plugins/Animation/MLDeformer/NearestNeighborModel/NearestNeighborModel.uplugin b/Engine/Plugins/Animation/MLDeformer/NearestNeighborModel/NearestNeighborModel.uplugin
--- a/Engine/Plugins/Animation/MLDeformer/NearestNeighborModel/NearestNeighborModel.uplugin
+++ b/Engine/Plugins/Animation/MLDeformer/NearestNeighborModel/NearestNeighborModel.uplugin
@@ -59,7 +59,6 @@
 			"Requirements":
 			[
 				"joblib==1.2.0",
-				"onnx==1.13.0",
 				"scipy==1.8.1",
 				"scikit-learn==1.2.1",
 				"threadpoolctl==3.1.0"
