NearestNeighborModel: Remove scipy dependency from pip requirements since it now lives in python foundation packages.

diff --git a/Engine/Plugins/Animation/MLDeformer/NearestNeighborModel/NearestNeighborModel.uplugin b/Engine/Plugins/Animation/MLDeformer/NearestNeighborModel/NearestNeighborModel.uplugin
--- a/Engine/Plugins/Animation/MLDeformer/NearestNeighborModel/NearestNeighborModel.uplugin
+++ b/Engine/Plugins/Animation/MLDeformer/NearestNeighborModel/NearestNeighborModel.uplugin
@@ -59,7 +59,6 @@
 			"Requirements":
 			[
 				"joblib==1.2.0",
-				"scipy==1.8.1",
 				"scikit-learn==1.2.1",
 				"threadpoolctl==3.1.0"
 			]
