NearestNeighborModel: add an option to extract geometry cache in nearest neighbor model kmeans tool

diff --git a/Engine/Plugins/Animation/MLDeformer/NearestNeighborModel/Source/NearestNeighborModelEditor/Private/NearestNeighborEditorHelpers.cpp b/Engine/Plugins/Animation/MLDeformer/NearestNeighborModel/Source/NearestNeighborModelEditor/Private/NearestNeighborEditorHelpers.cpp
--- a/Engine/Plugins/Animation/MLDeformer/NearestNeighborModel/Source/NearestNeighborModelEditor/Private/NearestNeighborEditorHelpers.cpp
+++ b/Engine/Plugins/Animation/MLDeformer/NearestNeighborModel/Source/NearestNeighborModelEditor/Private/NearestNeighborEditorHelpers.cpp
@@ -5,6 +5,10 @@
 #include "Animation/AnimSequence.h"
 #include "Animation/AttributesRuntime.h"
 #include "BonePose.h"
+#include "GeometryCache.h"
+#include "GeometryCacheConstantTopologyWriter.h"
+#include "GeometryCacheMeshData.h"
+#include "GeometryCacheTrack.h"
 #include "NearestNeighborModelHelpers.h"
 
 #define LOCTEXT_NAMESPACE "NearestNeighborEditorHelpers"
@@ -64,6 +68,16 @@ namespace UE::NearestNeighborModel::Private
 		FBlendedCurve OutCurve;
 		UE::Anim::FStackAttributeContainer TempAttributes;
 	};
+
+	int32 GetNumVertices(UGeometryCacheTrack& Track)
+	{
+		FGeometryCacheMeshData MeshData;
+		if (!Track.GetMeshDataAtSampleIndex(0, MeshData))
+		{
+			return 0;
+		}
+		return MeshData.Positions.Num(); 
+	}
 };
 
 void UNearestNeighborAnimStream::Init(USkeleton* InSkeleton)
@@ -124,6 +138,11 @@ bool UNearestNeighborAnimStream::AppendFrames(const UAnimSequence* Anim, TArray<
 	{
 		const int32 Frame = Frames[Index];
 		TArray<FTransform> BoneTransforms = AnimEval.GetBoneTransforms(Anim, Frame);
+		if (BoneTransforms.IsEmpty()) // Frame not found
+		{
+			UE_LOG(LogNearestNeighborModel, Error, TEXT("Frame %d not found in AnimSequence"), Frame);
+			return false;
+		}
 		check (BoneTransforms.Num() == NumBones);
 		for (int32 BoneIndex = 0; BoneIndex < NumBones; ++BoneIndex)
 		{
@@ -180,4 +199,102 @@ bool UNearestNeighborAnimStream::ToAnim(UAnimSequence* OutAnim) const
 	return true;
 }
 
+void UNearestNeighborGeometryCacheStream::Init(const UGeometryCache* InTemplateCache)
+{
+	if (!InTemplateCache)
+	{
+		UE_LOG(LogNearestNeighborModel, Error, TEXT("TemplateCache is None when initializing NearestNeighborGeometryCacheStream"));
+		return;
+	}
+	TemplateCache = InTemplateCache;
+	TemplateNumTracks = TemplateCache->Tracks.Num();
+	TemplateTrackNumVertices.SetNum(TemplateNumTracks);
+	for (int32 TrackIndex = 0; TrackIndex < TemplateNumTracks; ++TrackIndex)
+	{
+		const TObjectPtr<UGeometryCacheTrack> Track = TemplateCache->Tracks[TrackIndex];
+		const int32 NumVertices = UE::NearestNeighborModel::Private::GetNumVertices(*Track);
+		if (NumVertices == 0)
+		{
+			UE_LOG(LogNearestNeighborModel, Warning, TEXT("Track %d has no vertices"), TrackIndex);
+		}
+		TemplateTrackNumVertices[TrackIndex] = NumVertices;
+	}
+	TrackToFrameToPositions.Empty();
+	TrackToFrameToPositions.SetNum(TemplateNumTracks);
+}
+
+bool UNearestNeighborGeometryCacheStream::IsValid() const
+{
+	return TemplateCache != nullptr;
+}
+
+bool UNearestNeighborGeometryCacheStream::AppendFrames(const UGeometryCache* Cache, TArray<int32> Frames)
+{
+	if (!Cache || !IsValid())
+	{
+		return false;
+	}
+	const int32 NumCacheFrames = UE::NearestNeighborModel::FHelpers::GetNumFrames(Cache);
+	const int32 NumTracks = Cache->Tracks.Num();
+	if (NumTracks != TemplateNumTracks)
+	{
+		UE_LOG(LogNearestNeighborModel, Error, TEXT("Number of tracks in Cache (%d) is different from TemplateCache (%d)"), NumTracks, TemplateNumTracks);
+		return false;
+	}
+	for (int32 Frame : Frames)
+	{
+		if (Frame < 0 || Frame >= NumCacheFrames) // Frame not found
+		{
+			UE_LOG(LogNearestNeighborModel, Error, TEXT("Frame %d not found in GeometryCache"), Frame);
+			return false;
+		}
+		for (int32 TrackIndex = 0; TrackIndex < NumTracks; ++TrackIndex)
+		{
+			const TObjectPtr<UGeometryCacheTrack> Track = Cache->Tracks[TrackIndex];
+			FGeometryCacheMeshData MeshData;
+			if (!Track->GetMeshDataAtSampleIndex(Frame, MeshData))
+			{
+				UE_LOG(LogNearestNeighborModel, Error, TEXT("Frame %d cannot be retrieved in GeometryCache"), Frame);
+				return false;
+			}
+			TArray<FVector3f>& Positions = MeshData.Positions;
+			if (Positions.Num() != TemplateTrackNumVertices[TrackIndex])
+			{
+				UE_LOG(LogNearestNeighborModel, Error, TEXT("Number of vertices in Frame %d of Track %d (%d) is different from TemplateCache (%d)"), Frame, TrackIndex, Positions.Num(), TemplateTrackNumVertices[TrackIndex]);
+				return false;
+			}
+			TArray<TArray<FVector3f>>& FrameToPositions = TrackToFrameToPositions[TrackIndex];
+			FrameToPositions.Add(MoveTemp(Positions));
+		}
+	}
+	return true;
+}
+
+bool UNearestNeighborGeometryCacheStream::ToGeometryCache(UGeometryCache* OutCache)
+{
+	if (!OutCache)
+	{
+		UE_LOG(LogNearestNeighborModel, Error, TEXT("OutCache is None when converting NearestNeighborGeometryCacheStream to GeometryCache"));
+		return false;
+	}
+	if (!IsValid())
+	{
+		UE_LOG(LogNearestNeighborModel, Error, TEXT("TemplateCache is None when converting NearestNeighborGeometryCacheStream to GeometryCache"));
+		return false;
+	}
+	using UE::GeometryCacheHelpers::FGeometryCacheConstantTopologyWriter;
+	using UE::GeometryCacheHelpers::AddTrackWritersFromTemplateCache;
+	using FTrackWriter = FGeometryCacheConstantTopologyWriter::FTrackWriter;
+
+	FGeometryCacheConstantTopologyWriter Writer(*OutCache);
+	const int32 NumAdded = AddTrackWritersFromTemplateCache(Writer, *TemplateCache);
+	check(NumAdded == TemplateNumTracks);
+	for (int32 TrackIndex = 0; TrackIndex < TemplateNumTracks; ++TrackIndex)
+	{
+		FTrackWriter& TrackWriter = Writer.GetTrackWriter(TrackIndex);
+		TrackWriter.WriteAndClose(TrackToFrameToPositions[TrackIndex]);
+	}
+	return true;
+}
+
 #undef LOCTEXT_NAMESPACE
\ No newline at end of file
diff --git a/Engine/Plugins/Animation/MLDeformer/NearestNeighborModel/Source/NearestNeighborModelEditor/Private/NearestNeighborEditorHelpers.h b/Engine/Plugins/Animation/MLDeformer/NearestNeighborModel/Source/NearestNeighborModelEditor/Private/NearestNeighborEditorHelpers.h
--- a/Engine/Plugins/Animation/MLDeformer/NearestNeighborModel/Source/NearestNeighborModelEditor/Private/NearestNeighborEditorHelpers.h
+++ b/Engine/Plugins/Animation/MLDeformer/NearestNeighborModel/Source/NearestNeighborModelEditor/Private/NearestNeighborEditorHelpers.h
@@ -29,4 +29,29 @@ private:
 	TArray<TArray<FVector3f>> PosKeys;
 	TArray<TArray<FQuat4f>> RotKeys;
 	TArray<TArray<FVector3f>> ScaleKeys;
+};
+
+class UGeometryCache;
+UCLASS(Blueprintable)
+class NEARESTNEIGHBORMODELEDITOR_API UNearestNeighborGeometryCacheStream : public UObject
+{
+	GENERATED_BODY()
+public:
+	UFUNCTION(BlueprintCallable, Category = "Python")
+	void Init(const UGeometryCache* InTemplateCache);
+
+	UFUNCTION(BlueprintPure, Category = "Python")
+	bool IsValid() const;
+
+	UFUNCTION(BlueprintCallable, Category = "Python")
+	bool AppendFrames(const UGeometryCache* Cache, TArray<int32> Frames);
+
+	UFUNCTION(BlueprintPure, Category = "Python")
+	bool ToGeometryCache(UGeometryCache* OutCache);
+
+private:
+	const UGeometryCache* TemplateCache = nullptr;
+	int32 TemplateNumTracks = 0;
+	TArray<int32> TemplateTrackNumVertices;
+	TArray<TArray<TArray<FVector3f>>> TrackToFrameToPositions;
 };
\ No newline at end of file
diff --git a/Engine/Plugins/Animation/MLDeformer/NearestNeighborModel/Source/NearestNeighborModelEditor/Private/NearestNeighborTrainingModel.cpp b/Engine/Plugins/Animation/MLDeformer/NearestNeighborModel/Source/NearestNeighborModelEditor/Private/NearestNeighborTrainingModel.cpp
--- a/Engine/Plugins/Animation/MLDeformer/NearestNeighborModel/Source/NearestNeighborModelEditor/Private/NearestNeighborTrainingModel.cpp
+++ b/Engine/Plugins/Animation/MLDeformer/NearestNeighborModel/Source/NearestNeighborModelEditor/Private/NearestNeighborTrainingModel.cpp
@@ -29,11 +29,16 @@ UNearestNeighborModel* UNearestNeighborTrainingModel::GetNearestNeighborModel()
 	return GetCastModel();
 }
 
-int32 UNearestNeighborTrainingModel::GetNumFrames(const UAnimSequence* Anim) const
+int32 UNearestNeighborTrainingModel::GetNumFramesAnimSequence(const UAnimSequence* Anim) const
 {
 	return UE::NearestNeighborModel::FHelpers::GetNumFrames(Anim);
 }
 
+int32 UNearestNeighborTrainingModel::GetNumFramesGeometryCache(const UGeometryCache* GeometryCache) const
+{
+	return UE::NearestNeighborModel::FHelpers::GetNumFrames(GeometryCache);
+}
+
 const USkeleton* UNearestNeighborTrainingModel::GetModelSkeleton(const UMLDeformerModel* Model) const
 {
 	if (Model)
diff --git a/Engine/Plugins/Animation/MLDeformer/NearestNeighborModel/Source/NearestNeighborModelEditor/Private/Tools/NearestNeighborKMeansTool.cpp b/Engine/Plugins/Animation/MLDeformer/NearestNeighborModel/Source/NearestNeighborModelEditor/Private/Tools/NearestNeighborKMeansTool.cpp
--- a/Engine/Plugins/Animation/MLDeformer/NearestNeighborModel/Source/NearestNeighborModelEditor/Private/Tools/NearestNeighborKMeansTool.cpp
+++ b/Engine/Plugins/Animation/MLDeformer/NearestNeighborModel/Source/NearestNeighborModelEditor/Private/Tools/NearestNeighborKMeansTool.cpp
@@ -130,6 +130,28 @@ namespace UE::NearestNeighborModel
 					return FReply::Handled();
 				})
 			]
+			+ SHorizontalBox::Slot()
+			.AutoWidth()
+			[
+				SNew(SButton)
+				.Text(FText::FromString("New Geometry Cache"))
+				.HAlign(HAlign_Center)
+				.VAlign(VAlign_Center)
+				.OnClicked_Lambda([&Data]() -> FReply
+				{
+					if (UNearestNeighborKMeansData* KMeansData = Cast<UNearestNeighborKMeansData>(&Data))
+					{
+						if (KMeansData->NearestNeighborModelAsset)
+						{
+							const FString DefaultPath = FPackageName::GetLongPackagePath(KMeansData->NearestNeighborModelAsset->GetOutermost()->GetName());
+							const FString DefaultName = FString::Printf(TEXT("AS_Clusters_%s"), *KMeansData->NearestNeighborModelAsset->GetName());
+							UGeometryCache* Asset = Private::NewAssetDialog<UGeometryCache>(DefaultPath, DefaultName);
+							KMeansData->ExtractedCache = Asset;
+						}
+					}
+					return FReply::Handled();
+				})
+			]
 		]
 		+SVerticalBox::Slot()
 		.AutoHeight()
diff --git a/Engine/Plugins/Animation/MLDeformer/NearestNeighborModel/Source/NearestNeighborModelEditor/Private/Tools/NearestNeighborKMeansTool.h b/Engine/Plugins/Animation/MLDeformer/NearestNeighborModel/Source/NearestNeighborModelEditor/Private/Tools/NearestNeighborKMeansTool.h
--- a/Engine/Plugins/Animation/MLDeformer/NearestNeighborModel/Source/NearestNeighborModelEditor/Private/Tools/NearestNeighborKMeansTool.h
+++ b/Engine/Plugins/Animation/MLDeformer/NearestNeighborModel/Source/NearestNeighborModelEditor/Private/Tools/NearestNeighborKMeansTool.h
@@ -26,17 +26,33 @@ public:
 	UPROPERTY(VisibleAnywhere, Category = "Input")
 	TObjectPtr<const UMLDeformerAsset> NearestNeighborModelAsset;
 
+	/** Section used to generate clustered poses. */
 	UPROPERTY(EditAnywhere, BlueprintReadOnly, Category = "Input")
 	int32 SectionIndex = 0;
 	
-	UPROPERTY(EditAnywhere, BlueprintReadOnly, Category = "Input")
+	/** Number of clusters to be generated. */
+	UPROPERTY(EditAnywhere, BlueprintReadOnly, Category = "Input", meta = (ClampMin = 1))
 	int32 NumClusters = 10;
 
+	/** List of input poses (cannot be empty). */
 	UPROPERTY(EditAnywhere, BlueprintReadOnly, Category = "Input")
 	TArray<TObjectPtr<UAnimSequence>> InputPoses;
 
+	/** Whether to extract geometry cache at the same time. */
+	UPROPERTY(EditAnywhere, BlueprintReadOnly, Category = "Input")
+	bool bExtractGeometryCache = false;
+
+	/** List of input geometry caches (need to be the same size as InputPoses). */
+	UPROPERTY(EditAnywhere, BlueprintReadOnly, Category = "Input", meta = (EditCondition = "bExtractGeometryCache"))
+	TArray<TObjectPtr<UGeometryCache>> InputCaches;
+
+	/** Extracted poses. */
 	UPROPERTY(EditAnywhere, BlueprintReadOnly, Category = "Output")
 	TObjectPtr<UAnimSequence> ExtractedPoses;
+
+	/** Extracted geometry cache. */
+	UPROPERTY(EditAnywhere, BlueprintReadOnly, Category = "Output", meta = (EditCondition = "bExtractGeometryCache"))
+	TObjectPtr<UGeometryCache> ExtractedCache;
 };
 
 namespace UE::NearestNeighborModel
diff --git a/Engine/Plugins/Animation/MLDeformer/NearestNeighborModel/Source/NearestNeighborModelEditor/Public/NearestNeighborTrainingModel.h b/Engine/Plugins/Animation/MLDeformer/NearestNeighborModel/Source/NearestNeighborModelEditor/Public/NearestNeighborTrainingModel.h
--- a/Engine/Plugins/Animation/MLDeformer/NearestNeighborModel/Source/NearestNeighborModelEditor/Public/NearestNeighborTrainingModel.h
+++ b/Engine/Plugins/Animation/MLDeformer/NearestNeighborModel/Source/NearestNeighborModelEditor/Public/NearestNeighborTrainingModel.h
@@ -50,7 +50,10 @@ protected:
 
 private:
 	UFUNCTION(BlueprintPure, Category = "Python")
-	int32 GetNumFrames(const UAnimSequence* Anim) const;
+	int32 GetNumFramesAnimSequence(const UAnimSequence* Anim) const;
+
+	UFUNCTION(BlueprintPure, Category = "Python")
+	int32 GetNumFramesGeometryCache(const UGeometryCache* GeometryCache) const;
 	
 	UFUNCTION(BlueprintPure, Category = "Python")
 	const USkeleton* GetModelSkeleton(const UMLDeformerModel* Model) const;
diff --git a/Engine/Plugins/Runtime/GeometryCache/Source/GeometryCache/Private/GeometryCacheConstantTopologyWriter.cpp b/Engine/Plugins/Runtime/GeometryCache/Source/GeometryCache/Private/GeometryCacheConstantTopologyWriter.cpp
--- a/Engine/Plugins/Runtime/GeometryCache/Source/GeometryCache/Private/GeometryCacheConstantTopologyWriter.cpp
+++ b/Engine/Plugins/Runtime/GeometryCache/Source/GeometryCache/Private/GeometryCacheConstantTopologyWriter.cpp
@@ -196,6 +196,18 @@ namespace UE::GeometryCacheHelpers
 			const TArray<int32>& Map = LODModel.MeshToImportVertexMap;
 			return TArray<uint32>(reinterpret_cast<const uint32*>(Map.GetData()), Map.Num());
 		}
+
+		template<typename T>
+		TArray<T> GetSubArray(const TArray<T>& Array, int32 Start, int32 End)
+		{
+			if (Start >= End || Start >= Array.Num())
+			{
+				return TArray<T>();
+			}
+			Start = FMath::Clamp(Start, 0, Array.Num() - 1);
+			End = FMath::Clamp(End, 0, Array.Num());
+			return TArray<T>(Array.GetData() + Start, End - Start);
+		}
 	}
 
 	FGeometryCacheConstantTopologyWriter::FGeometryCacheConstantTopologyWriter(UGeometryCache& OutCache)
@@ -212,9 +224,9 @@ namespace UE::GeometryCacheHelpers
 	}
 	
 	using FTrackWriter = FGeometryCacheConstantTopologyWriter::FTrackWriter;
-	FTrackWriter& FGeometryCacheConstantTopologyWriter::AddTrackWriter()
+	FTrackWriter& FGeometryCacheConstantTopologyWriter::AddTrackWriter(FName TrackName)
 	{
-		return TrackWriters.Emplace_GetRef(*this);
+		return TrackWriters.Emplace_GetRef(*this, TrackName);
 	}
 
 	FTrackWriter& FGeometryCacheConstantTopologyWriter::GetTrackWriter(int32 TrackIndex)
@@ -227,7 +239,7 @@ namespace UE::GeometryCacheHelpers
 		return TrackWriters.Num();
 	}
 
-	void FGeometryCacheConstantTopologyWriter::AddMaterials(TArrayView<TObjectPtr<UMaterialInterface>> InMaterials)
+	void FGeometryCacheConstantTopologyWriter::AddMaterials(const TArray<TObjectPtr<UMaterialInterface>>& InMaterials)
 	{
 		if (!Cache.IsValid())
 		{
@@ -241,7 +253,7 @@ namespace UE::GeometryCacheHelpers
 		return Cache.IsValid() ? Cache->Materials.Num() : 0;
 	}
 
-	FGeometryCacheConstantTopologyWriter::FTrackWriter::FTrackWriter(FGeometryCacheConstantTopologyWriter& InOwner)
+	FGeometryCacheConstantTopologyWriter::FTrackWriter::FTrackWriter(FGeometryCacheConstantTopologyWriter& InOwner, FName TrackName)
 		: Owner(&InOwner) 
 	{
 		UGeometryCache* const CachePtr = Owner->Cache.Get();
@@ -250,7 +262,10 @@ namespace UE::GeometryCacheHelpers
 			return;
 		}
 		const FString BaseName = CachePtr->GetName();
-		const FName TrackName = MakeUniqueObjectName(CachePtr, UGeometryCacheTrackStreamable::StaticClass(), FName(BaseName + FString(TEXT("_Track"))));
+		if (TrackName.IsNone())
+		{
+			TrackName = MakeUniqueObjectName(CachePtr, UGeometryCacheTrackStreamable::StaticClass(), FName(BaseName + FString(TEXT("_Track"))));
+		}
 		Track = TStrongObjectPtr(NewObject<UGeometryCacheTrackStreamable>(CachePtr, TrackName, RF_Public));
 	}
 
@@ -444,5 +459,35 @@ namespace UE::GeometryCacheHelpers
 		Writer.AddMaterials(GetMaterialInterfaces(Asset));
 		return Writer.GetNumTracks() - 1;
 	}
+
+	int32 AddTrackWritersFromTemplateCache(FGeometryCacheConstantTopologyWriter& Writer, const UGeometryCache& TemplateCache)
+	{
+		int32 NumAddedTracks = 0;
+		for (TObjectPtr<UGeometryCacheTrack> Track : TemplateCache.Tracks)
+		{
+			if (!Track)
+			{
+				continue;
+			}
+			FGeometryCacheMeshData MeshData;
+			const bool bSuccess = Track->GetMeshDataAtSampleIndex(0, MeshData);
+			if (!bSuccess)
+			{
+				continue;
+			}
+
+			FTrackWriter& TrackWriter = Writer.AddTrackWriter(FName(*Track->GetName()));
+			TrackWriter.Indices = MeshData.Indices;
+			TrackWriter.UVs = MeshData.TextureCoordinates;
+			TrackWriter.Colors = MeshData.Colors;
+			TrackWriter.ImportedVertexNumbers = MeshData.ImportedVertexNumbers;
+			TrackWriter.BatchesInfo = MeshData.BatchesInfo;
+			NumAddedTracks++;
+		}
+
+		Writer.AddMaterials(TemplateCache.Materials);
+		return NumAddedTracks;
+	}
+
 #endif // WITH_EDITOR
 };
\ No newline at end of file
diff --git a/Engine/Plugins/Runtime/GeometryCache/Source/GeometryCache/Public/GeometryCacheConstantTopologyWriter.h b/Engine/Plugins/Runtime/GeometryCache/Source/GeometryCache/Public/GeometryCacheConstantTopologyWriter.h
--- a/Engine/Plugins/Runtime/GeometryCache/Source/GeometryCache/Public/GeometryCacheConstantTopologyWriter.h
+++ b/Engine/Plugins/Runtime/GeometryCache/Source/GeometryCache/Public/GeometryCacheConstantTopologyWriter.h
@@ -8,6 +8,7 @@
 
 struct FGeometryCacheMeshBatchInfo;
 struct FPackedNormal;
+class FSkeletalMeshLODModel;
 class UGeometryCache;
 class UGeometryCacheTrack;
 class UMaterialInterface; 
@@ -61,7 +62,7 @@ namespace UE::GeometryCacheHelpers
 
 		struct GEOMETRYCACHE_API FTrackWriter
 		{
-			FTrackWriter(FGeometryCacheConstantTopologyWriter& InOwner);
+			FTrackWriter(FGeometryCacheConstantTopologyWriter& InOwner, FName TrackName = FName());
 			~FTrackWriter() = default;
 
 			FTrackWriter(const FTrackWriter&) = delete;
@@ -98,10 +99,10 @@ namespace UE::GeometryCacheHelpers
 			FGeometryCacheConstantTopologyWriter* Owner = nullptr;
 		};
 
-		FTrackWriter& AddTrackWriter();
+		FTrackWriter& AddTrackWriter(FName TrackName = FName());
 		FTrackWriter& GetTrackWriter(int32 Index);
 		int32 GetNumTracks() const;
-		void AddMaterials(TArrayView<TObjectPtr<UMaterialInterface>> Materials);
+		void AddMaterials(const TArray<TObjectPtr<UMaterialInterface>>& Materials);
 		int32 GetNumMaterials() const;
 
 	private:
@@ -121,5 +122,12 @@ namespace UE::GeometryCacheHelpers
 	 * 	}
 	 */
 	GEOMETRYCACHE_API int32 AddTrackWriterFromSkinnedAsset(FGeometryCacheConstantTopologyWriter& Writer, const USkinnedAsset& Asset);
+
+	/**
+	 * @brief This will create multiple track writers and fill in the track writer's data (indices, UVs, materials .etc) from the template geometry cache. 
+	 * The number of track writers created equals to the number of tracks in the template geometry cache. 
+	 * @return The number of track writers created.
+	 */
+	GEOMETRYCACHE_API int32 AddTrackWritersFromTemplateCache(FGeometryCacheConstantTopologyWriter& Writer, const UGeometryCache& TemplateCache);
 #endif // WITH_EDITOR
 };
\ No newline at end of file
