[MLDeformer] Fixed a crash when switching between NMM and NN and back to NMM models.

diff --git a/Engine/Plugins/Animation/MLDeformer/NearestNeighborModel/Source/NearestNeighborModelEditor/Private/NearestNeighborEditorModel.cpp b/Engine/Plugins/Animation/MLDeformer/NearestNeighborModel/Source/NearestNeighborModelEditor/Private/NearestNeighborEditorModel.cpp
--- a/Engine/Plugins/Animation/MLDeformer/NearestNeighborModel/Source/NearestNeighborModelEditor/Private/NearestNeighborEditorModel.cpp
+++ b/Engine/Plugins/Animation/MLDeformer/NearestNeighborModel/Source/NearestNeighborModelEditor/Private/NearestNeighborEditorModel.cpp
@@ -45,14 +45,6 @@ namespace UE::NearestNeighborModel
 		}
 	};
 
-	FNearestNeighborEditorModel::~FNearestNeighborEditorModel()
-	{
-		if (NearestNeighborActor.Get())
-		{
-			EditorActors.Remove(NearestNeighborActor.Get());
-		}
-	}
-
 	UE::MLDeformer::FMLDeformerEditorModel* FNearestNeighborEditorModel::MakeInstance()
 	{
 		return new FNearestNeighborEditorModel();
@@ -92,9 +84,9 @@ namespace UE::NearestNeighborModel
 			return;
 		}
 
-		if (NearestNeighborActor.Get())
+		if (NearestNeighborActor)
 		{
-			EditorActors.Remove(NearestNeighborActor.Get());
+			EditorActors.Remove(NearestNeighborActor);
 			if (NearestNeighborActor->GetActor())
 			{
 				World->DestroyActor(NearestNeighborActor->GetActor(), true);
@@ -102,13 +94,13 @@ namespace UE::NearestNeighborModel
 		}
 		NearestNeighborActor = CreateNearestNeighborActor(World);
 		UpdateNearestNeighborActor(*NearestNeighborActor);
-		EditorActors.Add(NearestNeighborActor.Get());
+		EditorActors.Add(NearestNeighborActor);
 	}
 
 	void FNearestNeighborEditorModel::Tick(FEditorViewportClient* ViewportClient, float DeltaTime)
 	{
 		FMLDeformerEditorModel::Tick(ViewportClient, DeltaTime);
-		if (NearestNeighborActor.Get())
+		if (NearestNeighborActor)
 		{
 			UpdateNearestNeighborActor(*NearestNeighborActor);
 			NearestNeighborActor->Tick();
@@ -405,7 +397,7 @@ namespace UE::NearestNeighborModel
 		return Cast<UNearestNeighborModel>(Model.Get());
 	}
 
-	TUniquePtr<FNearestNeighborEditorModelActor> FNearestNeighborEditorModel::CreateNearestNeighborActor(UWorld* World) const
+	FNearestNeighborEditorModelActor* FNearestNeighborEditorModel::CreateNearestNeighborActor(UWorld* World) const
 	{
 		static FLinearColor LabelColor = FNearestNeighborModelEditorStyle::Get().GetColor("NearestNeighborModel.NearestNeighborActors.LabelColor");
 		static FLinearColor WireframeColor = FNearestNeighborModelEditorStyle::Get().GetColor("NearestNeighborModel.NearestNeighborActors.WireframeColor");
@@ -435,7 +427,7 @@ namespace UE::NearestNeighborModel
 		Settings.LabelText = LabelText;
 		Settings.bIsTrainingActor = false;
 
-		TUniquePtr<FNearestNeighborEditorModelActor> NewActor = MakeUnique<FNearestNeighborEditorModelActor>(Settings);
+		FNearestNeighborEditorModelActor* NewActor = new FNearestNeighborEditorModelActor(Settings);
 		NewActor->SetGeometryCacheComponent(GeomCacheComponent);
 		return NewActor;
 	}
diff --git a/Engine/Plugins/Animation/MLDeformer/NearestNeighborModel/Source/NearestNeighborModelEditor/Public/NearestNeighborEditorModel.h b/Engine/Plugins/Animation/MLDeformer/NearestNeighborModel/Source/NearestNeighborModelEditor/Public/NearestNeighborEditorModel.h
--- a/Engine/Plugins/Animation/MLDeformer/NearestNeighborModel/Source/NearestNeighborModelEditor/Public/NearestNeighborEditorModel.h
+++ b/Engine/Plugins/Animation/MLDeformer/NearestNeighborModel/Source/NearestNeighborModelEditor/Public/NearestNeighborEditorModel.h
@@ -28,8 +28,6 @@ namespace UE::NearestNeighborModel
 		using FSection = UNearestNeighborModelSection;
 		using FMLDeformerSampler = UE::MLDeformer::FMLDeformerSampler;
 
-		virtual ~FNearestNeighborEditorModel();
-
 		// We need to implement this static MakeInstance method.
 		static FMLDeformerEditorModel* MakeInstance();
 
@@ -82,10 +80,10 @@ namespace UE::NearestNeighborModel
 		void ResetMorphTargets();
 		void UpdateNearestNeighborIds();
 
-		TUniquePtr<FNearestNeighborEditorModelActor> CreateNearestNeighborActor(UWorld* World) const;
+		FNearestNeighborEditorModelActor* CreateNearestNeighborActor(UWorld* World) const;
 		void UpdateNearestNeighborActor(FNearestNeighborEditorModelActor& Actor) const;
 	
-		TUniquePtr<FNearestNeighborEditorModelActor> NearestNeighborActor;	// This should be only set in CreateActors().
+		FNearestNeighborEditorModelActor* NearestNeighborActor = nullptr;	// This should be only set in CreateActors() and is automatically deleted by the base class.
 		TUniquePtr<FVertexMapSelector> VertexMapSelector;
 		TUniquePtr<FVertVizSelector> VertVizSelector;
 	};
