NearestNeighborModel: improve base model

diff --git a/Engine/Plugins/Animation/MLDeformer/NearestNeighborModel/Source/NearestNeighborModel/Private/NearestNeighborModel.cpp b/Engine/Plugins/Animation/MLDeformer/NearestNeighborModel/Source/NearestNeighborModel/Private/NearestNeighborModel.cpp
--- a/Engine/Plugins/Animation/MLDeformer/NearestNeighborModel/Source/NearestNeighborModel/Private/NearestNeighborModel.cpp
+++ b/Engine/Plugins/Animation/MLDeformer/NearestNeighborModel/Source/NearestNeighborModel/Private/NearestNeighborModel.cpp
@@ -53,6 +53,7 @@ namespace UE::NearestNeighborModel
 		enum Type
 		{
 			BeforeCustomVersionWasAdded = 0,
+			AddTrainedBasis = 1,
 	
 			VersionPlusOne,
 			LatestVersion = VersionPlusOne - 1
@@ -1202,6 +1203,11 @@ void UNearestNeighborModel::PostLoad()
 		UpdateForInference();
 	PRAGMA_ENABLE_DEPRECATION_WARNINGS
 	}
+
+	if (IsBeforeTrainedBasisAdded())
+	{
+		bUsePCA = true;
+	}
 #endif
 }
 
@@ -1425,6 +1431,11 @@ bool UNearestNeighborModel::IsBeforeCustomVersionWasAdded() const
 	return Version < FNearestNeighborModelCustomVersion::BeforeCustomVersionWasAdded;
 }
 
+bool UNearestNeighborModel::IsBeforeTrainedBasisAdded() const
+{
+	using UE::NearestNeighborModel::FNearestNeighborModelCustomVersion;
+	return Version < FNearestNeighborModelCustomVersion::AddTrainedBasis;
+}
 
 UE::NearestNeighborModel::EOpFlag UNearestNeighborModel::CheckHiddenLayerDims()
 {
diff --git a/Engine/Plugins/Animation/MLDeformer/NearestNeighborModel/Source/NearestNeighborModel/Private/NearestNeighborModelInstance.cpp b/Engine/Plugins/Animation/MLDeformer/NearestNeighborModel/Source/NearestNeighborModel/Private/NearestNeighborModelInstance.cpp
--- a/Engine/Plugins/Animation/MLDeformer/NearestNeighborModel/Source/NearestNeighborModel/Private/NearestNeighborModelInstance.cpp
+++ b/Engine/Plugins/Animation/MLDeformer/NearestNeighborModel/Source/NearestNeighborModel/Private/NearestNeighborModelInstance.cpp
@@ -390,7 +390,7 @@ void UNearestNeighborModelInstance::RunNearestNeighborModel(float DeltaTime, flo
 		{
 			const float W = OutputView[MorphIndex] * ModelWeight;
 			using UE::NearestNeighborModel::Private::UpdateWeightWithDecay;
-			UpdateWeightWithDecay(WeightData->Weights[MorphIndex + 1], PreviousWeights[MorphIndex + 1], W, DecayCoeff);
+			UpdateWeightWithDecay(WeightData->Weights[MorphIndex + 1], PreviousWeights[MorphIndex + 1], W, 0.0f);
 		}
 
 		int32 NeighborOffset = NumNetworkWeights + 1;
diff --git a/Engine/Plugins/Animation/MLDeformer/NearestNeighborModel/Source/NearestNeighborModel/Public/NearestNeighborModel.h b/Engine/Plugins/Animation/MLDeformer/NearestNeighborModel/Source/NearestNeighborModel/Public/NearestNeighborModel.h
--- a/Engine/Plugins/Animation/MLDeformer/NearestNeighborModel/Source/NearestNeighborModel/Public/NearestNeighborModel.h
+++ b/Engine/Plugins/Animation/MLDeformer/NearestNeighborModel/Source/NearestNeighborModel/Public/NearestNeighborModel.h
@@ -399,6 +399,15 @@ public:
 	UFUNCTION(BlueprintPure, Category = "Nearest Neighbor Model")
 	int32 GetEarlyStopEpochs() const { return EarlyStopEpochs; }
 
+	UFUNCTION(BlueprintPure, Category = "Nearest Neighbor Model")
+	float GetRegularizationFactor() const { return RegularizationFactor; }
+
+	UFUNCTION(BlueprintPure, Category = "Nearest Neighbor Model")
+	float GetSmoothLossBeta() const { return SmoothLossBeta; }
+
+	UFUNCTION(BlueprintPure, Category = "Nearest Neighbor Model")
+	bool DoesUsePCA() const { return bUsePCA; }
+
 	UFUNCTION(BlueprintPure, Category = "Nearest Neighbor Model")
 	FString GetModelDir() const;
 	
@@ -447,6 +456,7 @@ public:
 	void UpdateNetworkOutputDim();
 
 	bool IsBeforeCustomVersionWasAdded() const;
+	bool IsBeforeTrainedBasisAdded() const;
 	const TArray<float>& GetVertexWeightSum() const;
 
 	static FName GetInputDimPropertyName() { return GET_MEMBER_NAME_CHECKED(UNearestNeighborModel, InputDim); }
@@ -460,6 +470,7 @@ public:
 	static FName GetNearestNeighborOffsetWeightPropertyName() { return GET_MEMBER_NAME_CHECKED(UNearestNeighborModel, NearestNeighborOffsetWeight); }
 	static FName GetUseFileCachePropertyName() { return GET_MEMBER_NAME_CHECKED(UNearestNeighborModel, bUseFileCache); }
 	static FName GetFileCacheDirectoryPropertyName() { return GET_MEMBER_NAME_CHECKED(UNearestNeighborModel, FileCacheDirectory); }
+	static FName GetUsePCAPropertyName() { return GET_MEMBER_NAME_CHECKED(UNearestNeighborModel, bUsePCA); }
 	static FName GetUseDualQuaternionDeltasPropertyName() { return GET_MEMBER_NAME_CHECKED(UNearestNeighborModel, bUseDualQuaternionDeltas); }
 	static FName GetDecayFactorPropertyName() { return GET_MEMBER_NAME_CHECKED(UNearestNeighborModel, DecayFactor); }
 	static FName GetUseRBFPropertyName() { return GET_MEMBER_NAME_CHECKED(UNearestNeighborModel, bUseRBF); }
@@ -488,12 +499,20 @@ protected:
 	UPROPERTY(EditAnywhere, BlueprintReadWrite, Category = "Training Settings", AdvancedDisplay, meta = (ClampMin = "1"))
 	int32 BatchSize = 256;
 
+	/** The regularization factor. Higher values can help generate more sparse morph targets, but can also lead to visual artifacts.  */
+	UPROPERTY(EditAnywhere, BlueprintReadWrite, Category = "Training Settings", AdvancedDisplay, meta = (ClampMin = "0", EditCondition = "!bUsePCA"))
+	float RegularizationFactor = 1.0f;
+	
+	/** The beta parameter in the smooth L1 loss function, which describes below which absolute error to use a squared term. If the error is above or equal to this beta value, it will use the L1 loss. */
+	UPROPERTY(EditAnywhere, BlueprintReadWrite, AdvancedDisplay, Category = "Training Settings", meta = (ClampMin = "0.0", EditCondition = "!bUsePCA"))
+	float SmoothLossBeta = 1.0f;
+
 	/** The size of the step when optimizing the network. */
 	UPROPERTY(EditAnywhere, BlueprintReadWrite, Category = "Training Settings", AdvancedDisplay, meta = (ClampMin = "0.000001", ClampMax = "1.0"))
 	float LearningRate = 0.001f;
 
 	/** The number of epochs to stop training if there is no improvement in accuracy. */
-	UPROPERTY(EditAnywhere, BlueprintReadWrite, Category = "Training Settings", AdvancedDisplay, meta = (ClampMin = "1"))
+	UPROPERTY(EditAnywhere, BlueprintReadWrite, Category = "Training Settings", AdvancedDisplay, meta = (ClampMin = "1", EditCondition = "bUsePCA"))
 	int32 EarlyStopEpochs = 100;
 
 	/** Whether to cache intermediate results on disk. CAUTION: failing to manually clear caches could cause unexpected results. */
@@ -533,6 +552,9 @@ PRAGMA_ENABLE_DEPRECATION_WARNINGS
 	/** The max input values observed throughout the entire training set. This is used to clamp the input value at inference time.  The values are set in python. */
 	UPROPERTY(BlueprintReadWrite, Category = "Network IO")
 	TArray<float> InputsMax;
+	
+	UPROPERTY(EditAnywhere, BlueprintReadWrite, Category = "Nearest Neighbor Settings")
+	bool bUsePCA = false;
 
 	/** Whether to use dual quaternion deltas. If false, LBS deltas will be used. */
 	UPROPERTY(EditAnywhere, BlueprintReadWrite, Category = "Nearest Neighbor Settings")
