NearestNeighborModel:
- Fix bugs with multi-section using learned basis
- Remove PCA in function & variable names
- Fix vertex attribute update bug
- Move training data to GPU
- Parallel delta computation
- Do not invalidate model when data is set to None
- Add "Clear References" button

diff --git a/Engine/Plugins/Animation/MLDeformer/MLDeformerFramework/Source/MLDeformerFramework/Public/MLDeformerVizSettings.h b/Engine/Plugins/Animation/MLDeformer/MLDeformerFramework/Source/MLDeformerFramework/Public/MLDeformerVizSettings.h
--- a/Engine/Plugins/Animation/MLDeformer/MLDeformerFramework/Source/MLDeformerFramework/Public/MLDeformerVizSettings.h
+++ b/Engine/Plugins/Animation/MLDeformer/MLDeformerFramework/Source/MLDeformerFramework/Public/MLDeformerVizSettings.h
@@ -73,6 +73,7 @@ public:
 	void SetTrainingFrameNumber(int32 FrameNumber)			{ TrainingFrameNumber = FrameNumber; }
 	void SetTestingFrameNumber(int32 FrameNumber)			{ TestingFrameNumber = FrameNumber; }
 	void SetWeight(float InWeight)							{ Weight = InWeight; }
+	void SetTestAnimSequence(UAnimSequence* InAnim)  		{ TestAnimSequence = InAnim; }
 
 	FVector GetMeshSpacingOffsetVector() const				{ return FVector(MeshSpacing, 0.0f, 0.0f); }
 	float GetMeshSpacing() const							{ return MeshSpacing; }
diff --git a/Engine/Plugins/Animation/MLDeformer/NearestNeighborModel/Source/NearestNeighborModel/Private/NearestNeighborModel.cpp b/Engine/Plugins/Animation/MLDeformer/NearestNeighborModel/Source/NearestNeighborModel/Private/NearestNeighborModel.cpp
--- a/Engine/Plugins/Animation/MLDeformer/NearestNeighborModel/Source/NearestNeighborModel/Private/NearestNeighborModel.cpp
+++ b/Engine/Plugins/Animation/MLDeformer/NearestNeighborModel/Source/NearestNeighborModel/Private/NearestNeighborModel.cpp
@@ -355,18 +355,23 @@ int32 UNearestNeighborModelSection::GetAssetNumNeighbors() const
 	return FMath::Min(FHelpers::GetNumFrames(GetNeighborPoses()), FHelpers::GetNumFrames(GetNeighborMeshes()));
 }
 
-void UNearestNeighborModelSection::SetPCAData(const TArray<float>& InVertexMean, const TArray<float>& InPCABasis)
+void UNearestNeighborModelSection::SetBasisData(const TArray<float>& InVertexMean, const TArray<float>& InBasis)
 {
+	const int32 NumBasis = GetNumBasis();
+	check(NumBasis > 0);
+	check(InVertexMean.Num() == NumVertices * 3);
+	check(InBasis.Num() == NumVertices * 3 * NumBasis);
 	VertexMean = InVertexMean;
-	PCABasis = InPCABasis;
+	Basis = InBasis;
 	InvalidateInference();
 }
 
 void UNearestNeighborModelSection::SetNeighborData(const TArray<float>& InAssetNeighborCoeffs, const TArray<float>& InAssetNeighborOffsets)
 {
-	check(NumPCACoeffs > 0);
-	const int32 TryNumNeighbors = InAssetNeighborCoeffs.Num() / NumPCACoeffs;
-	check(TryNumNeighbors * NumPCACoeffs == InAssetNeighborCoeffs.Num());
+	const int32 NumBasis = GetNumBasis();
+	check(NumBasis > 0);
+	const int32 TryNumNeighbors = InAssetNeighborCoeffs.Num() / NumBasis;
+	check(TryNumNeighbors * NumBasis == InAssetNeighborCoeffs.Num());
 	check(TryNumNeighbors * NumVertices * 3 == InAssetNeighborOffsets.Num());
 	AssetNeighborCoeffs = InAssetNeighborCoeffs;
 	AssetNeighborOffsets = InAssetNeighborOffsets;
@@ -383,9 +388,9 @@ const TArray<float>& UNearestNeighborModelSection::GetVertexWeights() const
 	return VertexWeights;
 }
 
-const TArray<float>& UNearestNeighborModelSection::GetPCABasis() const
+const TArray<float>& UNearestNeighborModelSection::GetBasis() const
 {
-	return PCABasis;
+	return Basis;
 }
 
 const TArray<float>& UNearestNeighborModelSection::GetVertexMean() const
@@ -398,6 +403,11 @@ const TArray<float>& UNearestNeighborModelSection::GetAssetNeighborCoeffs() cons
 	return AssetNeighborCoeffs;
 }
 
+bool UNearestNeighborModelSection::DoesUsePCA() const
+{
+	return Model && Model->DoesUsePCA();
+}
+
 const TArray<float>& UNearestNeighborModelSection::GetAssetNeighborOffsets() const
 {
 	return AssetNeighborOffsets;
@@ -426,12 +436,13 @@ FMLDeformerGeomCacheTrainingInputAnim* UNearestNeighborModelSection::GetInputAni
 
 bool UNearestNeighborModelSection::UpdateRuntimeNeighbors()
 {
-	if (NumPCACoeffs <= 0)
+	const int32 NumBasis = GetNumBasis();
+	if (NumBasis <= 0)
 	{
 		return false;
 	}
-	const int32 AssetNumNeighbors = AssetNeighborCoeffs.Num() / NumPCACoeffs;
-	if (AssetNumNeighbors * NumPCACoeffs != AssetNeighborCoeffs.Num())
+	const int32 AssetNumNeighbors = AssetNeighborCoeffs.Num() / NumBasis;
+	if (AssetNumNeighbors * NumBasis != AssetNeighborCoeffs.Num())
 	{
 		return false;
 	}
@@ -442,11 +453,11 @@ bool UNearestNeighborModelSection::UpdateRuntimeNeighbors()
 		return true;
 	}
 	RuntimeNeighborCoeffs.Reset();
-	RuntimeNeighborCoeffs.Reserve(RuntimeNumNeighbors * NumPCACoeffs);	
+	RuntimeNeighborCoeffs.Reserve(RuntimeNumNeighbors * NumBasis);	
 	for (int32 Index = 0; Index < IncludedFrames.Num(); ++Index)
 	{
 		const int32 Frame = IncludedFrames[Index];
-		TConstArrayView<float> AssetCoeffs(AssetNeighborCoeffs.GetData() + Frame * NumPCACoeffs,  NumPCACoeffs);
+		TConstArrayView<float> AssetCoeffs(AssetNeighborCoeffs.GetData() + Frame * NumBasis,  NumBasis);
 		Algo::Copy(AssetCoeffs, RuntimeNeighborCoeffs);
 	}
 	AssetNeighborIndexMap = IncludedFrames;
@@ -475,7 +486,7 @@ bool UNearestNeighborModelSection::GetRuntimeNeighborOffsets(TArray<float>& OutN
 }
 #endif
 
-int32 UNearestNeighborModelSection::GetNumPCACoeffs() const
+int32 UNearestNeighborModelSection::GetNumBasis() const
 {
 	return NumPCACoeffs;
 }
@@ -545,10 +556,10 @@ UE::NearestNeighborModel::EOpFlag UNearestNeighborModelSection::UpdateForInferen
 		return Result;
 	}
 
-	if (!IsPCAValid())
+	if (!IsBasisValid())
 	{
-		UE_LOG(LogNearestNeighborModel, Error, TEXT("PCA data is invalid. Please re-train your model"));
-		ResetPCAData();
+		UE_LOG(LogNearestNeighborModel, Error, TEXT("Basis is invalid. Please re-train your model"));
+		ResetBasisData();
 		return EOpFlag::Error;
 	}
 	if (!UpdateRuntimeNeighbors())
@@ -558,7 +569,7 @@ UE::NearestNeighborModel::EOpFlag UNearestNeighborModelSection::UpdateForInferen
 	}
 	if (!IsNearestNeighborValid())
 	{
-		UE_LOG(LogNearestNeighborModel, Error, TEXT("Nearest neighbor data is invalid. Please re-train your model to update PCA data or re-update your model to update nearest neighbor data."));
+		UE_LOG(LogNearestNeighborModel, Error, TEXT("Nearest neighbor data is invalid. Please re-train your model to update basis data or re-update your model to update nearest neighbor data."));
 		ResetNearestNeighborData();
 		return EOpFlag::Error;
 	}
@@ -653,10 +664,21 @@ void UNearestNeighborModelSection::PostEditChangeProperty(FPropertyChangedEvent&
 		InvalidateTraining();
 		UpdateVertexWeights();
 	}
-	// Sometimes we need to set geometry cache to None before checking in. Do not invalidate training in this case.
-	if (Property->GetFName() == GET_MEMBER_NAME_CHECKED(UNearestNeighborModelSection, NeighborPoses) ||
-		// Property->GetFName() == GET_MEMBER_NAME_CHECKED(UNearestNeighborModelSection, NeighborMeshes) ||
-		Property->GetFName() == GET_MEMBER_NAME_CHECKED(UNearestNeighborModelSection, ExcludedFrames))
+	if (Property->GetFName() == GET_MEMBER_NAME_CHECKED(UNearestNeighborModelSection, NeighborPoses))
+	{
+		if (NeighborPoses)
+		{
+			InvalidateInference();
+		}
+	}
+	if (Property->GetFName() == GET_MEMBER_NAME_CHECKED(UNearestNeighborModelSection, NeighborMeshes))
+	{
+		if (NeighborMeshes)
+		{
+			InvalidateInference();
+		}
+	}
+	if (Property->GetFName() == GET_MEMBER_NAME_CHECKED(UNearestNeighborModelSection, ExcludedFrames))
 	{
 		InvalidateInference();
 	}
@@ -664,6 +686,12 @@ void UNearestNeighborModelSection::PostEditChangeProperty(FPropertyChangedEvent&
 	Model->PostEditChangeProperty(PropertyChangedEvent);
 }
 
+void UNearestNeighborModelSection::SetNumBasis(const int32 InNumBasis)
+{
+	NumPCACoeffs = InNumBasis;
+	InvalidateTraining();
+}
+
 int32 UNearestNeighborModelSection::GetMeshIndex() const
 {
 	return MeshIndex;
@@ -676,27 +704,26 @@ void UNearestNeighborModelSection::SetMeshIndex(int32 Index)
 	Model->InvalidateTrainingModelOnly();
 }
 
-bool UNearestNeighborModelSection::IsPCAValid() const
+bool UNearestNeighborModelSection::IsBasisValid() const
 {
-	return NumPCACoeffs >= 1 && NumVertices > 0 
-	&& VertexMean.Num() == NumVertices * 3 && PCABasis.Num() == NumVertices * 3 * NumPCACoeffs;
+	return VertexMean.Num() == NumVertices * 3 && Basis.Num() == NumVertices * 3 * GetNumBasis();
 }
 
-bool UNearestNeighborModelSection::IsPCAEmpty() const
+bool UNearestNeighborModelSection::IsBasisEmpty() const
 {
-	return VertexMean.IsEmpty() && PCABasis.IsEmpty();
+	return VertexMean.IsEmpty() && Basis.IsEmpty();
 }
 
-void UNearestNeighborModelSection::ResetPCAData()
+void UNearestNeighborModelSection::ResetBasisData()
 {
 	VertexMean.Reset();
-	PCABasis.Reset();
+	Basis.Reset();
 }
 
 bool UNearestNeighborModelSection::IsNearestNeighborValid() const
 {
 	return IsNearestNeighborEmpty() 
-	|| (RuntimeNeighborCoeffs.Num() == RuntimeNumNeighbors * NumPCACoeffs && NumVertices > 0);
+	|| (RuntimeNeighborCoeffs.Num() == RuntimeNumNeighbors * GetNumBasis() && NumVertices > 0);
 }
 
 bool UNearestNeighborModelSection::IsNearestNeighborEmpty() const
@@ -709,7 +736,7 @@ void UNearestNeighborModelSection::Reset()
 	NumVertices = 0;
 	RuntimeNumNeighbors = 0;
 	VertexMap.Reset();
-	ResetPCAData();
+	ResetBasisData();
 	ResetNearestNeighborData();
 	bIsReadyForTraining = false;
 	bIsReadyForInference = false;
@@ -749,10 +776,14 @@ UE::NearestNeighborModel::EOpFlag UNearestNeighborModelSection::UpdateVertexWeig
 
 	NumVertices = VertexMap.Num();
 	check(VertexWeights.Num() == NumVertices);
-	if (NumVertices <= NumPCACoeffs)
+
+	if (DoesUsePCA())
 	{
-		UE_LOG(LogNearestNeighborModel, Error, TEXT("NumVertices %d needs to be larger than NumPCACoeffs %d."), NumVertices, NumPCACoeffs);
-		Result |= EOpFlag::Error;
+		if (NumVertices <= GetNumBasis())
+		{
+			UE_LOG(LogNearestNeighborModel, Error, TEXT("NumVertices %d needs to be larger than NumBasis %d."), NumVertices, GetNumBasis());
+			Result |= EOpFlag::Error;
+		}
 	}
 	return Result;
 }
@@ -1009,7 +1040,7 @@ void UNearestNeighborModelSection::InitFromClothPartData(FClothPartData& InPart)
 	VertexMapString = IntegersToFormattedString(VertexMap);
 	MeshIndex = INDEX_NONE;
 	NumPCACoeffs = InPart.PCACoeffNum;
-	PCABasis = MoveTemp(InPart.PCABasis);
+	Basis = MoveTemp(InPart.PCABasis);
 	VertexMean = MoveTemp(InPart.VertexMean);
 	NumVertices = VertexMap.Num();
 	RuntimeNumNeighbors = InPart.NumNeighbors;
@@ -1078,14 +1109,6 @@ void UNearestNeighborModel::PostEditChangeProperty(FPropertyChangedEvent& Proper
 	{
 		InvalidateTraining();
 	}
-	if (PropertyChangedEvent.GetMemberPropertyName() == TEXT("TrainingInputAnims"))
-	{
-		// Sometimes we need to set geometry cache to None before checking in. Do not invalidate training in this case.
-		if (Property->GetFName() != FMLDeformerGeomCacheTrainingInputAnim::GetGeomCachePropertyName())
-		{
-			InvalidateTraining();
-		}
-	}
 
 	if (Property->GetFName() == UNearestNeighborModel::GetBoneIncludeListPropertyName() ||
 		Property->GetFName() == UNearestNeighborModel::GetCurveIncludeListPropertyName())
@@ -1118,6 +1141,13 @@ void UNearestNeighborModel::PostEditChangeProperty(FPropertyChangedEvent& Proper
 			UpdateNetworkOutputDim();
 		}
 	}
+	if (Property->GetFName() == UNearestNeighborModel::GetUsePCAPropertyName() ||
+		Property->GetFName() == UNearestNeighborModel::GetNumBasisPerSectionPropertyName())
+	{
+		UpdateSectionNumBasis();
+		UpdateNetworkOutputDim();
+		InvalidateTraining();
+	}
 	if (Property->GetFName() == UNearestNeighborModel::GetUseFileCachePropertyName() ||
 		Property->GetFName() == UNearestNeighborModel::GetFileCacheDirectoryPropertyName())
 	{
@@ -1125,6 +1155,45 @@ void UNearestNeighborModel::PostEditChangeProperty(FPropertyChangedEvent& Proper
 	}
 	Super::PostEditChangeProperty(PropertyChangedEvent);
 }
+
+void UNearestNeighborModel::PostEditChangeChainProperty(FPropertyChangedChainEvent& PropertyChangedEvent) 
+{
+	Super::PostEditChangeProperty(PropertyChangedEvent);
+	const FProperty* Property = PropertyChangedEvent.Property;
+	if (!Property)
+	{
+		return;
+	}
+	const FProperty* MemberProperty = PropertyChangedEvent.PropertyChain.GetActiveMemberNode()->GetValue();
+	if (!MemberProperty)
+	{
+		return;
+	}
+
+	const FString MemberPropertyName = MemberProperty->GetName();
+	if (MemberPropertyName == TEXT("TrainingInputAnims"))
+	{
+		bool bInvalidate = true;
+		const int32 ArrayIndex = PropertyChangedEvent.GetArrayIndex(MemberPropertyName);
+		const TArray<FMLDeformerGeomCacheTrainingInputAnim> InputAnims = GetTrainingInputAnims();
+		if (ArrayIndex >= 0 && ArrayIndex < InputAnims.Num())
+		{
+			if (Property->GetFName() == FMLDeformerTrainingInputAnim::GetAnimSequencePropertyName() && !InputAnims[ArrayIndex].GetAnimSequence())
+			{
+				bInvalidate = false;
+			}
+			if (Property->GetFName() == FMLDeformerGeomCacheTrainingInputAnim::GetGeomCachePropertyName() && !InputAnims[ArrayIndex].GetGeometryCache())
+			{
+				bInvalidate = false;
+			}
+		}
+		if (bInvalidate)
+		{
+			InvalidateTraining();
+		}
+	}
+}
+
 #endif
 
 UMLDeformerInputInfo* UNearestNeighborModel::CreateInputInfo()
@@ -1231,14 +1300,26 @@ const TArray<int32>& UNearestNeighborModel::GetPCACoeffStarts() const
 	return PCACoeffStarts;
 }
 
-int32 UNearestNeighborModel::GetTotalNumPCACoeffs() const
+int32 UNearestNeighborModel::GetTotalNumBasis() const
 {
-	int32 Sum = 0;
-	for(const FSection* Section : Sections)
+	if (DoesUsePCA())
 	{
-		Sum += Section->GetNumPCACoeffs();
+		int32 Sum = 0;
+		for(const FSection* Section : Sections)
+		{
+			Sum += Section->GetNumBasis();
+		}
+		return Sum;
 	}
-	return Sum;
+	else
+	{
+		return NumBasisPerSection;
+	}
+}
+
+int32 UNearestNeighborModel::GetNumBasisPerSection() const
+{
+	return NumBasisPerSection;
 }
 
 int32 UNearestNeighborModel::GetTotalNumNeighbors() const
@@ -1322,6 +1403,10 @@ UNearestNeighborModelSection* UNearestNeighborModel::OnSectionAdded(int32 NewInd
 		return nullptr;
 	}
 	Section->SetModel(this);
+	if (!DoesUsePCA())
+	{
+		Section->SetNumBasis(NumBasisPerSection);
+	}
 	Sections[NewIndex] = Section;
 	return Section;
 }
@@ -1411,7 +1496,7 @@ void UNearestNeighborModel::UpdatePCACoeffStarts()
 	for(const FSection* Section : Sections)
 	{
 		PCACoeffStarts.Add(Acc);
-		Acc += Section->GetNumPCACoeffs();
+		Acc += Section->GetNumBasis();
 	}
 }
 
@@ -1422,9 +1507,31 @@ void UNearestNeighborModel::UpdateNetworkInputDim()
 
 void UNearestNeighborModel::UpdateNetworkOutputDim()
 {
-	OutputDim = GetTotalNumPCACoeffs();
+	OutputDim = GetTotalNumBasis();
 }
 
+void UNearestNeighborModel::ClearReferences()
+{
+	for (FSection* Section : Sections)
+	{
+		if (Section)
+		{
+			Section->ClearReferences();
+		}
+	}
+	for (FMLDeformerGeomCacheTrainingInputAnim& InputAnim : GetTrainingInputAnims())
+	{
+		InputAnim.SetAnimSequence(nullptr);
+		InputAnim.SetGeometryCache(nullptr);
+	}
+	UMLDeformerVizSettings* const Viz = GetVizSettings();
+	if (Viz)
+	{
+		Viz->SetTestAnimSequence(nullptr);
+	}
+}
+
+
 bool UNearestNeighborModel::IsBeforeCustomVersionWasAdded() const
 {
 	using UE::NearestNeighborModel::FNearestNeighborModelCustomVersion;
@@ -1488,6 +1595,17 @@ void UNearestNeighborModel::SetFileCacheDirectory(const FString& InFileCacheDire
 	FileCacheDirectory = InFileCacheDirectory;
 }
 
+void UNearestNeighborModel::UpdateSectionNumBasis()
+{
+	for (FSection* Section : Sections)
+	{
+		if (Section)
+		{
+			Section->SetNumBasis(NumBasisPerSection);
+		}
+	}
+}
+
 void UNearestNeighborModel::UpdateCachedDeltasTimestamp()
 {
 	CachedDeltasTimestamp = UE::NearestNeighborModel::Private::GetLatestTimeStamp(GetCachedDeltasPaths());
@@ -1534,6 +1652,12 @@ void UNearestNeighborModel::NormalizeVertexWeights()
 	VertexWeightSum.Reset();
 }
 
+void UNearestNeighborModelSection::ClearReferences()
+{
+	NeighborPoses = nullptr;
+	NeighborMeshes = nullptr;
+}
+
 UE::NearestNeighborModel::EOpFlag UNearestNeighborModel::UpdateForTraining()
 {
 	if (IsReadyForTraining())
@@ -1542,7 +1666,10 @@ UE::NearestNeighborModel::EOpFlag UNearestNeighborModel::UpdateForTraining()
 	}
 	using namespace UE::NearestNeighborModel;
 	EOpFlag Result = EOpFlag::Success;
-	UpdatePCACoeffStarts();
+	if (DoesUsePCA())
+	{
+		UpdatePCACoeffStarts();
+	}
 	UpdateNetworkInputDim();
 	if (InputDim == 0)
 	{
@@ -1552,7 +1679,7 @@ UE::NearestNeighborModel::EOpFlag UNearestNeighborModel::UpdateForTraining()
 	UpdateNetworkOutputDim();
 	if (OutputDim == 0)
 	{
-		UE_LOG(LogNearestNeighborModel, Error, TEXT("OutputDim is 0. Please create at least one section and use non-zero PCA coeff counts."));
+		UE_LOG(LogNearestNeighborModel, Error, TEXT("OutputDim is 0. Please create at least one section."));
 		return EOpFlag::Error;
 	}
 	Result |= CheckHiddenLayerDims();
diff --git a/Engine/Plugins/Animation/MLDeformer/NearestNeighborModel/Source/NearestNeighborModel/Private/NearestNeighborModelInstance.cpp b/Engine/Plugins/Animation/MLDeformer/NearestNeighborModel/Source/NearestNeighborModel/Private/NearestNeighborModelInstance.cpp
--- a/Engine/Plugins/Animation/MLDeformer/NearestNeighborModel/Source/NearestNeighborModel/Private/NearestNeighborModelInstance.cpp
+++ b/Engine/Plugins/Animation/MLDeformer/NearestNeighborModel/Source/NearestNeighborModel/Private/NearestNeighborModelInstance.cpp
@@ -206,7 +206,7 @@ void UNearestNeighborModelInstance::InitInstanceData(int32 NumMorphWeights)
 		{
 			return;
 		}
-		NumMorphWeights = 1 + NearestNeighborModel->GetTotalNumPCACoeffs() + NearestNeighborModel->GetTotalNumNeighbors();
+		NumMorphWeights = 1 + NearestNeighborModel->GetTotalNumBasis() + NearestNeighborModel->GetTotalNumNeighbors();
 	}
 	PreviousWeights.SetNumZeroed(NumMorphWeights);
 	DistanceBuffer.SetNumZeroed(NumMorphWeights);
@@ -266,16 +266,16 @@ namespace UE::NearestNeighborModel::Private
 		return DistanceSquared;
 	}
 
-	void ComputeNeighborDistances(TArrayView<float> OutDistances, const UNearestNeighborModel::FSection& Section, TConstArrayView<float> PCACoeffs)
+	void ComputeNeighborDistances(TArrayView<float> OutDistances, const UNearestNeighborModel::FSection& Section, TConstArrayView<float> Coeffs)
 	{
-		const int32 NumCoeffs = Section.GetNumPCACoeffs();
+		const int32 NumBasis = Section.GetNumBasis();
 		const int32 NumNeighbors = Section.GetRuntimeNumNeighbors();
 		check(OutDistances.Num() == NumNeighbors);
 		TConstArrayView<float> NeighborCoeffs = Section.GetNeighborCoeffs();
 		for (int32 NeighborId = 0; NeighborId < NumNeighbors; ++NeighborId)
 		{
-			TConstArrayView<float> Coeffs(NeighborCoeffs.GetData() + NeighborId * NumCoeffs, NumCoeffs);
-			OutDistances[NeighborId] = ComputeDistanceSquared(Coeffs, PCACoeffs);
+			TConstArrayView<float> SingleCoeffs(NeighborCoeffs.GetData() + NeighborId * NumBasis, NumBasis);
+			OutDistances[NeighborId] = ComputeDistanceSquared(Coeffs, SingleCoeffs);
 		}
 	}
 
@@ -402,20 +402,28 @@ void UNearestNeighborModelInstance::RunNearestNeighborModel(float DeltaTime, flo
 				continue;
 			}
 
-			TConstArrayView<int32> PCACoeffStarts = NearestNeighborModel->GetPCACoeffStarts();
-			if (!PCACoeffStarts.IsValidIndex(SectionIndex))
+			int32 CoeffStart = 0;
+			if (NearestNeighborModel->DoesUsePCA())
 			{
-				continue;
+				TConstArrayView<int32> PCACoeffStarts = NearestNeighborModel->GetPCACoeffStarts();
+				if (!PCACoeffStarts.IsValidIndex(SectionIndex))
+				{
+					continue;
+				}
+				CoeffStart = PCACoeffStarts[SectionIndex];
+			}
+			else
+			{
+				CoeffStart = 0;
 			}
-			const int32 PCACoeffStart = PCACoeffStarts[SectionIndex];
-			const int32 NumCoeffs = Section.GetNumPCACoeffs();
-			TConstArrayView<float> PCACoeffs(OutputView.GetData() + PCACoeffStart, NumCoeffs);
+			const int32 NumCoeffs = Section.GetNumBasis();
+			TConstArrayView<float> Coeffs(OutputView.GetData() + CoeffStart, NumCoeffs);
 
 			const int32 NumNeighbors = Section.GetRuntimeNumNeighbors();
 
 			TArrayView<float> SquaredDistances(DistanceBuffer.GetData() + NeighborOffset, NumNeighbors);
 			using UE::NearestNeighborModel::Private::ComputeNeighborDistances;
-			ComputeNeighborDistances(SquaredDistances, Section, PCACoeffs);
+			ComputeNeighborDistances(SquaredDistances, Section, Coeffs);
 
 			TArrayView<float> SectionMorphWeights(WeightData->Weights.GetData() + NeighborOffset, NumNeighbors);
 			TArrayView<float> SectionPreviousWeights(PreviousWeights.GetData() + NeighborOffset, NumNeighbors);
diff --git a/Engine/Plugins/Animation/MLDeformer/NearestNeighborModel/Source/NearestNeighborModel/Public/NearestNeighborModel.h b/Engine/Plugins/Animation/MLDeformer/NearestNeighborModel/Source/NearestNeighborModel/Public/NearestNeighborModel.h
--- a/Engine/Plugins/Animation/MLDeformer/NearestNeighborModel/Source/NearestNeighborModel/Public/NearestNeighborModel.h
+++ b/Engine/Plugins/Animation/MLDeformer/NearestNeighborModel/Source/NearestNeighborModel/Public/NearestNeighborModel.h
@@ -100,7 +100,9 @@ public:
 	DECLARE_MULTICAST_DELEGATE_TwoParams(FPropertyChangedDelegate, FPropertyChangedEvent&, UNearestNeighborModelSection&);
 	using EOpFlag = UE::NearestNeighborModel::EOpFlag;
 
-	int32 GetNumPCACoeffs() const;
+	UFUNCTION(BlueprintPure, Category = "Section")
+	int32 GetNumBasis() const;
+
 	int32 GetNumVertices() const;
 	int32 GetRuntimeNumNeighbors() const;
 
@@ -112,6 +114,8 @@ public:
 #if WITH_EDITOR
 	virtual void PostEditChangeProperty(FPropertyChangedEvent& PropertyChangedEvent) override;
 	
+	void SetNumBasis(const int32 InNumBasis);
+
 	UFUNCTION(BlueprintPure, Category = "Section")
 	int32 GetAssetNumNeighbors() const;
 	
@@ -128,10 +132,12 @@ public:
 
 	UFUNCTION(BlueprintPure, Category = "Section")
 	const TArray<int32>& GetVertexMap() const;
+	
+	UFUNCTION(BlueprintPure, Category = "Section")
 	const TArray<float>& GetVertexWeights() const;
 	
 	UFUNCTION(BlueprintPure, Category = "Section")
-	const TArray<float>& GetPCABasis() const;
+	const TArray<float>& GetBasis() const;
 	
 	UFUNCTION(BlueprintPure, Category = "Section")
 	const TArray<float>& GetVertexMean() const;
@@ -139,6 +145,9 @@ public:
 	UFUNCTION(BlueprintPure, Category = "Section")
 	const TArray<float>& GetAssetNeighborCoeffs() const;
 
+	UFUNCTION(BlueprintPure, Category = "Section")
+	bool DoesUsePCA() const;
+
 	const TArray<float>& GetAssetNeighborOffsets() const;
 	const TArray<int32>& GetExcludedFrames() const;
 	const TArray<int32>& GetAssetNeighborIndexMap() const;
@@ -163,10 +172,14 @@ public:
 	const TArray<FName>& GetBoneNames() const;
 	void SetBoneNames(const TArray<FName>& InBoneNames);
 
+
 	// Do not call this function directly. Call UNearestNeighborModel::NormalizeVertexWeights() instead.
 	EOpFlag NormalizeVertexWeights();
 
-	static FName GetNumPCACoeffsPropertyName() { return GET_MEMBER_NAME_CHECKED(UNearestNeighborModelSection, NumPCACoeffs); }
+	void ClearReferences();
+
+
+	static FName GetNumBasisPropertyName() { return GET_MEMBER_NAME_CHECKED(UNearestNeighborModelSection, NumPCACoeffs); }
 	static FName GetVertexMapStringPropertyName() { return GET_MEMBER_NAME_CHECKED(UNearestNeighborModelSection, VertexMapString); }
 	static FName GetNeighborPosesPropertyName() { return GET_MEMBER_NAME_CHECKED(UNearestNeighborModelSection, NeighborPoses); }
 	static FName GetNeighborMeshesPropertyName() { return GET_MEMBER_NAME_CHECKED(UNearestNeighborModelSection, NeighborMeshes); }
@@ -227,15 +240,15 @@ protected:
 	UPROPERTY()
 	TArray<float> VertexWeights;
 
-	/** Flattened array of PCA basis. The shape of PCA basis is (PCACoeffNum, NumVertices * 3)  */
+	/** Flattened array of basis. The shape of basis is (CoeffNum, NumVertices * 3)  */
 	UPROPERTY()
-	TArray<float> PCABasis;
+	TArray<float> Basis;
 	
-	/** The vertex mean used in PCA computation. This array has a size of NumVertices * 3 */
+	/** The vertex mean on the shape. This array has a size of NumVertices * 3 */
 	UPROPERTY()
 	TArray<float> VertexMean;
 
-	/** Flattened array of neighbor PCA coefficients before excluding frames. The shape of this array is (NumNeighbors, NumPCACoeffs) */
+	/** Flattened array of neighbor coefficients before excluding frames. The shape of this array is (NumNeighbors, NumCoeffs) */
 	UPROPERTY()
 	TArray<float> AssetNeighborCoeffs;
 	
@@ -249,7 +262,7 @@ protected:
 #endif
 
 	/** Number of PCA coefficients for this section. */
-	UPROPERTY(EditAnywhere, BlueprintReadWrite, meta = (ClampMin = "1", DefaultValue = "64"), Category = "Section")
+	UPROPERTY(EditAnywhere, BlueprintReadWrite, meta = (ClampMin = "1", DefaultValue = "64", EditCondition = "DoesUsePCA()"), Category = "Section", DisplayName = "Num Basis") // TODO: rename to NumBasis
 	int32 NumPCACoeffs = 64;
 
 	/** Number of vertices in this section */
@@ -274,19 +287,19 @@ private:
 	bool bIsReadyForInference = false;
 
 #if WITH_EDITOR
-	bool IsPCAValid() const;
-	bool IsPCAEmpty() const;
+	bool IsBasisValid() const;
+	bool IsBasisEmpty() const;
 	bool IsNearestNeighborValid() const;
 	bool IsNearestNeighborEmpty() const;
 
 	UFUNCTION(BlueprintCallable, Category = "Section")
-	void SetPCAData(const TArray<float>& InVertexMean, const TArray<float>& InPCABasis);
+	void SetBasisData(const TArray<float>& InVertexMean, const TArray<float>& InBasis);
 	
 	UFUNCTION(BlueprintCallable, Category = "Section")
 	void SetNeighborData(const TArray<float>& InNeighborCoeffs, const TArray<float>& InNeighborOffsets);
 
 	void Reset();
-	void ResetPCAData();
+	void ResetBasisData();
 	void ResetNearestNeighborData();
 	EOpFlag UpdateVertexWeightsFromText();
 	EOpFlag UpdateVertexWeightsSelectedBones();
@@ -302,15 +315,15 @@ private:
 
 /**
  * The nearest neighbor model.
- * This model contains the PCA basis of the vertex deltas and a small set of meshes for nearest neighbor search. 
- * Given a new pose, the pre-trained neural network first predicts the PCA coefficients of the vertex deltas. 
- * Then this model uses the predicted PCA coeffcients to find a nearest neighbor in the small dataset.
+ * This model contains the linear basis of the vertex deltas and a small set of meshes for nearest neighbor search. 
+ * Given a new pose, the pre-trained neural network first predicts the coefficients of the vertex deltas. 
+ * Then this model uses the predicted coeffcients to find a nearest neighbor in the small dataset.
  * The total vertex delta is computed by
- * 		vertex_delta = mean_delta + pca_basis * pca_coeff + nearest_neighbor_delta
+ * 		vertex_delta = mean_delta + basis * coeff + nearest_neighbor_delta
  * To prevent popping, a time filtering is applied on predicted vertex deltas. The vertex delta at time t is computed by  
  * 		vertex_delta(t) = decay_factor * vertex_delta(t-1) + (1 - decay_factor) * vertex_delta 
  * The mesh can be separated into several sections (e.g. shirt, pants...). The nearest neighbor search is carried out separately for each section. 
- * The pca basis and the nearest neighbor data are compressed into morph targets.
+ * The basis and the nearest neighbor data are compressed into morph targets.
  */
 UCLASS()
 class NEARESTNEIGHBORMODEL_API UNearestNeighborModel
@@ -330,6 +343,7 @@ public:
 	virtual void GetAssetRegistryTags(TArray<FAssetRegistryTag>& OutTags) const override;
 #if WITH_EDITOR
 	virtual void PostEditChangeProperty(FPropertyChangedEvent& PropertyChangedEvent) override;
+	virtual void PostEditChangeChainProperty(FPropertyChangedChainEvent& PropertyChangedEvent) override;
 #endif
 	// ~END UObject overrides.
 
@@ -358,7 +372,8 @@ public:
 	const TArray<int32>& GetPCACoeffStarts() const;
 
 	UFUNCTION(BlueprintPure, Category = "Nearest Neighbor Model")
-	int32 GetTotalNumPCACoeffs() const;
+	int32 GetTotalNumBasis() const;
+	int32 GetNumBasisPerSection() const;
 
 	int32 GetTotalNumNeighbors() const;
 	float GetDecayFactor() const;
@@ -377,6 +392,9 @@ public:
 	bool DoesUseRBF() const;
 	float GetRBFSigma() const;
 
+	UFUNCTION(BlueprintPure, Category = "Nearest Neighbor Model")
+	bool DoesUsePCA() const { return bUsePCA; }
+
 #if WITH_EDITOR
 	UFUNCTION(BlueprintPure, Category = "Nearest Neighbor Model")
 	int32 GetInputDim() const { return InputDim; }
@@ -405,9 +423,6 @@ public:
 	UFUNCTION(BlueprintPure, Category = "Nearest Neighbor Model")
 	float GetSmoothLossBeta() const { return SmoothLossBeta; }
 
-	UFUNCTION(BlueprintPure, Category = "Nearest Neighbor Model")
-	bool DoesUsePCA() const { return bUsePCA; }
-
 	UFUNCTION(BlueprintPure, Category = "Nearest Neighbor Model")
 	FString GetModelDir() const;
 	
@@ -455,6 +470,8 @@ public:
 	void UpdateNetworkInputDim();
 	void UpdateNetworkOutputDim();
 
+	void ClearReferences();
+
 	bool IsBeforeCustomVersionWasAdded() const;
 	bool IsBeforeTrainedBasisAdded() const;
 	const TArray<float>& GetVertexWeightSum() const;
@@ -471,6 +488,7 @@ public:
 	static FName GetUseFileCachePropertyName() { return GET_MEMBER_NAME_CHECKED(UNearestNeighborModel, bUseFileCache); }
 	static FName GetFileCacheDirectoryPropertyName() { return GET_MEMBER_NAME_CHECKED(UNearestNeighborModel, FileCacheDirectory); }
 	static FName GetUsePCAPropertyName() { return GET_MEMBER_NAME_CHECKED(UNearestNeighborModel, bUsePCA); }
+	static FName GetNumBasisPerSectionPropertyName() { return GET_MEMBER_NAME_CHECKED(UNearestNeighborModel, NumBasisPerSection); }
 	static FName GetUseDualQuaternionDeltasPropertyName() { return GET_MEMBER_NAME_CHECKED(UNearestNeighborModel, bUseDualQuaternionDeltas); }
 	static FName GetDecayFactorPropertyName() { return GET_MEMBER_NAME_CHECKED(UNearestNeighborModel, DecayFactor); }
 	static FName GetUseRBFPropertyName() { return GET_MEMBER_NAME_CHECKED(UNearestNeighborModel, bUseRBF); }
@@ -493,7 +511,7 @@ protected:
 
 	/** Max number of cycles iterated through the training set. */
 	UPROPERTY(EditAnywhere, BlueprintReadWrite, Category = "Training Settings", AdvancedDisplay, meta = (ClampMin = "1"))
-	int32 NumEpochs = 10000;
+	int32 NumEpochs = 2500;
 
 	/** Number of data samples processed together as a group in a single pass. */
 	UPROPERTY(EditAnywhere, BlueprintReadWrite, Category = "Training Settings", AdvancedDisplay, meta = (ClampMin = "1"))
@@ -553,9 +571,14 @@ PRAGMA_ENABLE_DEPRECATION_WARNINGS
 	UPROPERTY(BlueprintReadWrite, Category = "Network IO")
 	TArray<float> InputsMax;
 	
+	/** Whether to use pre-computed PCA basis. If false, basis will be learned at training time. */
 	UPROPERTY(EditAnywhere, BlueprintReadWrite, Category = "Nearest Neighbor Settings")
 	bool bUsePCA = false;
 
+	/** The number of basis used in each section. Only editable when UsePCA is false. */
+	UPROPERTY(EditAnywhere, BlueprintReadWrite, Category = "Nearest Neighbor Settings", meta = (EditCondition = "!bUsePCA"))
+	int32 NumBasisPerSection = 128;
+
 	/** Whether to use dual quaternion deltas. If false, LBS deltas will be used. */
 	UPROPERTY(EditAnywhere, BlueprintReadWrite, Category = "Nearest Neighbor Settings")
 	bool bUseDualQuaternionDeltas = true;
@@ -590,6 +613,7 @@ private:
 	EOpFlag CheckHiddenLayerDims();
 	void UpdateInputMultipliers();
 
+	void UpdateSectionNumBasis();
 	void UpdateCachedDeltasTimestamp();
 	void UpdateCachedPCATimestamp();
 	void UpdateCachedNetworkTimestamp();
diff --git a/Engine/Plugins/Animation/MLDeformer/NearestNeighborModel/Source/NearestNeighborModelEditor/Private/NearestNeighborEditorModel.cpp b/Engine/Plugins/Animation/MLDeformer/NearestNeighborModel/Source/NearestNeighborModelEditor/Private/NearestNeighborEditorModel.cpp
--- a/Engine/Plugins/Animation/MLDeformer/NearestNeighborModel/Source/NearestNeighborModelEditor/Private/NearestNeighborEditorModel.cpp
+++ b/Engine/Plugins/Animation/MLDeformer/NearestNeighborModel/Source/NearestNeighborModelEditor/Private/NearestNeighborEditorModel.cpp
@@ -13,6 +13,7 @@
 #include "MLDeformerAsset.h"
 #include "MLDeformerComponent.h"
 #include "MLDeformerEditorStyle.h"
+#include "MLDeformerEditorToolkit.h"
 #include "NearestNeighborGeomCacheSampler.h"
 #include "NearestNeighborModel.h"
 #include "NearestNeighborModelInputInfo.h"
@@ -131,10 +132,10 @@ namespace UE::NearestNeighborModel
 			return ETrainingResult::FailPythonError;
 		}
 		const int32 NumFrames = GetNumFramesForTraining();
-		const int32 NumPCACoeffs = NearestNeighborModel->GetTotalNumPCACoeffs();
-		if (NumFrames <= NumPCACoeffs)
+		const int32 NumBasis = NearestNeighborModel->GetTotalNumBasis();
+		if (NumFrames <= NumBasis)
 		{
-			UE_LOG(LogNearestNeighborModel, Error, TEXT("Training frames (%d) must be greater than the number of PCA coefficients (%d)"), NumFrames, NumPCACoeffs);
+			UE_LOG(LogNearestNeighborModel, Error, TEXT("Training frames (%d) must be greater than the number of basis (%d)"), NumFrames, NumBasis);
 			return ETrainingResult::FailPythonError;
 		}
 		return TrainModel<UNearestNeighborTrainingModel>(this);
@@ -500,6 +501,12 @@ namespace UE::NearestNeighborModel
 		{
 			return UpdateResult;
 		}
+		UpdateIsReadyForTrainingState();
+		if (!IsReadyForTraining())
+		{
+			UE_LOG(LogNearestNeighborModel, Error, TEXT("Model is not ready for training. Please check training data or reload MLDeformer editor."));
+			return EOpFlag::Error;
+		}
 		UpdateResult |= UpdateNearestNeighborData();
 		if (OpFlag::HasError(UpdateResult))
 		{
@@ -545,6 +552,23 @@ namespace UE::NearestNeighborModel
 		}
 	}
 
+	void FNearestNeighborEditorModel::ClearReferences()
+	{
+		UNearestNeighborModel* const NearestNeighborModel = GetCastModel();
+		if (NearestNeighborModel)
+		{
+			NearestNeighborModel->ClearReferences();
+		}
+		if (Editor->GetModelDetailsView())
+		{
+			Editor->GetModelDetailsView()->ForceRefresh();
+		}
+		if (Editor->GetVizSettingsDetailsView())
+		{
+			Editor->GetVizSettingsDetailsView()->ForceRefresh();
+		}
+	}
+
 	EOpFlag FNearestNeighborEditorModel::CheckNetwork()
 	{
 		const UNearestNeighborModel* const NearestNeighborModel = GetCastModel();
@@ -560,18 +584,20 @@ namespace UE::NearestNeighborModel
 		}
 	
 		const int32 NumNetworkWeights = NearestNeighborModel->GetNumNetworkOutputs();
-		const int32 NumPCACoeffs = NearestNeighborModel->GetTotalNumPCACoeffs();
-		if (NumNetworkWeights != NumPCACoeffs)
+		const int32 NumBasis = NearestNeighborModel->GetTotalNumBasis();
+		if (NumNetworkWeights != NumBasis)
 		{
-			UE_LOG(LogNearestNeighborModel, Error, TEXT("Network output dimension %d is not equal to number of morph targets %d. Network needs to be re-trained."), NumNetworkWeights, NumPCACoeffs);
+			UE_LOG(LogNearestNeighborModel, Error, TEXT("Network output dimension %d does not equal to total number of basis %d. Network needs to be re-trained."), NumNetworkWeights, NumBasis);
 			return EOpFlag::Error;
 		}
+
 		return EOpFlag::Success;
 	}
 
 	EOpFlag FNearestNeighborEditorModel::UpdateNearestNeighborData()
 	{
-		if (!GetCastModel())
+		const UNearestNeighborModel* const NearestNeighborModel = GetCastModel();
+		if (!NearestNeighborModel)
 		{
 			return EOpFlag::Error;
 		}
@@ -580,6 +606,12 @@ namespace UE::NearestNeighborModel
 			UE_LOG(LogNearestNeighborModel, Error, TEXT("Network is not trained. Nearest neighbor data cannot be updated."));
 			return EOpFlag::Error;
 		}
+		if (NearestNeighborModel->GetNumNetworkOutputs() != NearestNeighborModel->GetTotalNumBasis())
+		{
+			UE_LOG(LogNearestNeighborModel, Error, TEXT("Network output dimension %d does not equal to total number of basis %d. Please re-train network."), NearestNeighborModel->GetNumNetworkOutputs(), NearestNeighborModel->GetTotalNumBasis());
+			return EOpFlag::Error;
+		}
+		
 		UNearestNeighborTrainingModel *TrainingModel = FHelpers::NewDerivedObject<UNearestNeighborTrainingModel>();
 		if (!TrainingModel)
 		{
@@ -629,9 +661,9 @@ namespace UE::NearestNeighborModel
 	
 		TArray<FVector3f> Deltas;
 	
-		const int32 TotalNumPCACoeffs = NearestNeighborModel->GetTotalNumPCACoeffs();
+		const int32 TotalNumBasis = NearestNeighborModel->GetTotalNumBasis();
 		const int32 TotalNumNeighbors = NearestNeighborModel->GetTotalNumNeighbors();
-		Deltas.SetNumZeroed((1 + TotalNumPCACoeffs + TotalNumNeighbors) * NumBaseMeshVerts);
+		Deltas.SetNumZeroed((1 + TotalNumBasis + TotalNumNeighbors) * NumBaseMeshVerts);
 	
 		using Private::AddFloatArrayToDeltaArray;
 		int32 MorphOffset = 1;
@@ -650,16 +682,23 @@ namespace UE::NearestNeighborModel
 			check(VertexWeights.Num() == Section.GetNumVertices());
 			TArrayView<FVector3f> MeanDeltas(Deltas.GetData(), NumBaseMeshVerts);
 			AddFloatArrayToDeltaArray(Section.GetVertexMean(), VertexMap, VertexWeights, MeanDeltas);
-			const int32 NumPCACoeffs = Section.GetNumPCACoeffs();
+			const int32 NumBasis = Section.GetNumBasis();
 			const int32 SectionNumVerts = Section.GetNumVertices();
 			check(VertexMap.Num() == SectionNumVerts);
-			for (int32 Index = 0; Index < NumPCACoeffs; ++Index)
+			for (int32 Index = 0; Index < NumBasis; ++Index)
 			{
-				TConstArrayView<float> BasisFloats(Section.GetPCABasis().GetData() + Index * SectionNumVerts * 3, SectionNumVerts * 3);
+				TConstArrayView<float> BasisFloats(Section.GetBasis().GetData() + Index * SectionNumVerts * 3, SectionNumVerts * 3);
 				TArrayView<FVector3f> BasisDeltas(Deltas.GetData() + (MorphOffset + Index) * NumBaseMeshVerts, NumBaseMeshVerts);
 				AddFloatArrayToDeltaArray(BasisFloats, VertexMap, VertexWeights, BasisDeltas);
 			}
-			MorphOffset += NumPCACoeffs;
+			if (NearestNeighborModel->DoesUsePCA())
+			{
+				MorphOffset += NumBasis;
+			}
+		}
+		if (!NearestNeighborModel->DoesUsePCA())
+		{
+			MorphOffset += TotalNumBasis;
 		}
 	
 		for (int32 SectionIndex = 0; SectionIndex < NearestNeighborModel->GetNumSections(); ++SectionIndex)
@@ -689,6 +728,7 @@ namespace UE::NearestNeighborModel
 				TArrayView<FVector3f> NeighborDeltas(Deltas.GetData() + (MorphOffset + Index) * NumBaseMeshVerts, NumBaseMeshVerts);
 				AddFloatArrayToDeltaArray(NeighborOffsets, VertexMap, VertexWeights, NeighborDeltas);
 			}
+			MorphOffset += SectionNumNeighbors;
 		}
 	
 		if (Deltas.Num() == 0)
diff --git a/Engine/Plugins/Animation/MLDeformer/NearestNeighborModel/Source/NearestNeighborModelEditor/Private/NearestNeighborGeomCacheSampler.cpp b/Engine/Plugins/Animation/MLDeformer/NearestNeighborModel/Source/NearestNeighborModelEditor/Private/NearestNeighborGeomCacheSampler.cpp
--- a/Engine/Plugins/Animation/MLDeformer/NearestNeighborModel/Source/NearestNeighborModelEditor/Private/NearestNeighborGeomCacheSampler.cpp
+++ b/Engine/Plugins/Animation/MLDeformer/NearestNeighborModel/Source/NearestNeighborModelEditor/Private/NearestNeighborGeomCacheSampler.cpp
@@ -5,6 +5,7 @@
 #include "Animation/AnimSequence.h"
 #include "Animation/AnimSingleNodeInstance.h"
 #include "Animation/DebugSkelMeshComponent.h"
+#include "Async/ParallelFor.h"
 #include "BoneWeights.h"
 #include "Engine/SkeletalMesh.h"
 #include "GeometryCache.h"
@@ -160,30 +161,40 @@ namespace UE::NearestNeighborModel
 
 			TArrayView<const int32> OffsetAndCount = PartOffsetAndCountRef.Get(MeshMapping.MeshIndex);
 			const int32 VertexOffset = OffsetAndCount[0];
-			const int32 VertexCount = OffsetAndCount[1];
+			const int32 NumVertices = OffsetAndCount[1];
 			
-			// Calculate the vertex deltas.
-			for (int32 VertexIndex = 0; VertexIndex < VertexCount; ++VertexIndex)
+			constexpr int32 BatchSize = 500;
+			const int32 NumBatches = (NumVertices / BatchSize) + 1;
+			ParallelFor(NumBatches, [&](int32 BatchIndex)
 			{
-				const int32 SkinnedVertexIndex = VertexOffset + VertexIndex;
-				const int32 GeomCacheVertexIndex = MeshMapping.SkelMeshToTrackVertexMap[VertexIndex];
-				if (GeomCacheMeshData.Positions.IsValidIndex(GeomCacheVertexIndex))
+				const int32 StartVertex = BatchIndex * BatchSize;
+				if (StartVertex >= NumVertices || VertexDeltas.IsEmpty())
 				{
-					FVector3f Delta = FVector3f::ZeroVector;
-					const int32 ArrayIndex = 3 * SkinnedVertexIndex;
-					const int32 RenderVertexIndex = MeshMapping.ImportedVertexToRenderVertexMap[VertexIndex];
-					if (RenderVertexIndex != INDEX_NONE)
+					return;
+				}
+				const int32 NumVertsInBatch = (StartVertex + BatchSize) < NumVertices ? BatchSize : FMath::Max(NumVertices - StartVertex, 0);
+				for (int32 VertexIndex = StartVertex; VertexIndex < StartVertex + NumVertsInBatch; ++VertexIndex)
+				{
+					const int32 SkinnedVertexIndex = VertexOffset + VertexIndex;
+					const int32 GeomCacheVertexIndex = MeshMapping.SkelMeshToTrackVertexMap[VertexIndex];
+					if (GeomCacheMeshData.Positions.IsValidIndex(GeomCacheVertexIndex))
 					{
-						const FVector3f SkinnedVertexPos = SkinnedVertexPositions[SkinnedVertexIndex];
-						const FVector3f GeomCacheVertexPos = (FVector3f)AlignmentTransform.TransformPosition((FVector)GeomCacheMeshData.Positions[GeomCacheVertexIndex]);
-						const FVector3f WorldDelta = GeomCacheVertexPos - SkinnedVertexPos;
-						Delta = CalcDualQuaternionDelta(RenderVertexIndex, WorldDelta, LODRenderData, *SkinWeightBuffer);
+						FVector3f Delta = FVector3f::ZeroVector;
+						const int32 ArrayIndex = 3 * SkinnedVertexIndex;
+						const int32 RenderVertexIndex = MeshMapping.ImportedVertexToRenderVertexMap[VertexIndex];
+						if (RenderVertexIndex != INDEX_NONE)
+						{
+							const FVector3f SkinnedVertexPos = SkinnedVertexPositions[SkinnedVertexIndex];
+							const FVector3f GeomCacheVertexPos = (FVector3f)AlignmentTransform.TransformPosition((FVector)GeomCacheMeshData.Positions[GeomCacheVertexIndex]);
+							const FVector3f WorldDelta = GeomCacheVertexPos - SkinnedVertexPos;
+							Delta = CalcDualQuaternionDelta(RenderVertexIndex, WorldDelta, LODRenderData, *SkinWeightBuffer);
+						}
+						VertexDeltas[ArrayIndex] = Delta.X;
+						VertexDeltas[ArrayIndex + 1] = Delta.Y;
+						VertexDeltas[ArrayIndex + 2] = Delta.Z;
 					}
-					VertexDeltas[ArrayIndex] = Delta.X;
-					VertexDeltas[ArrayIndex + 1] = Delta.Y;
-					VertexDeltas[ArrayIndex + 2] = Delta.Z;
 				}
-			}
+			});
 		}
 	}
 
diff --git a/Engine/Plugins/Animation/MLDeformer/NearestNeighborModel/Source/NearestNeighborModelEditor/Private/NearestNeighborModelDetails.cpp b/Engine/Plugins/Animation/MLDeformer/NearestNeighborModel/Source/NearestNeighborModelEditor/Private/NearestNeighborModelDetails.cpp
--- a/Engine/Plugins/Animation/MLDeformer/NearestNeighborModel/Source/NearestNeighborModelEditor/Private/NearestNeighborModelDetails.cpp
+++ b/Engine/Plugins/Animation/MLDeformer/NearestNeighborModel/Source/NearestNeighborModelEditor/Private/NearestNeighborModelDetails.cpp
@@ -216,6 +216,8 @@ namespace UE::NearestNeighborModel
 		{
 			return;
 		}
+		NearestNeighborCategoryBuilder->AddProperty(UNearestNeighborModel::GetUsePCAPropertyName());
+		NearestNeighborCategoryBuilder->AddProperty(UNearestNeighborModel::GetNumBasisPerSectionPropertyName());
 		NearestNeighborCategoryBuilder->AddProperty(UNearestNeighborModel::GetUseDualQuaternionDeltasPropertyName());
 		NearestNeighborCategoryBuilder->AddProperty(UNearestNeighborModel::GetDecayFactorPropertyName());
 		NearestNeighborCategoryBuilder->AddProperty(UNearestNeighborModel::GetNearestNeighborOffsetWeightPropertyName());
@@ -377,6 +379,27 @@ namespace UE::NearestNeighborModel
 				})
 			]
 		];
+
+		StatusCategoryBuilder->AddCustomRow(FText::FromString("StatusUpdate")).WholeRowContent()
+		[
+			SNew(SBox)
+			.Padding(2, 2)
+			.MaxDesiredWidth(200.f)
+			[
+				SNew(SButton)
+				.Text(FText::FromString(TEXT("Clear References")))
+				.ToolTipText(LOCTEXT("ClearReferencesTooltip", "Clears all animation sequences and geometry caches for asset validation."))
+				.HAlign(HAlign_Center)
+				.OnClicked_Lambda([this]
+				{
+					if (GetCastEditorModel())
+					{
+						GetCastEditorModel()->ClearReferences();
+					}
+					return FReply::Handled();
+				})
+			]
+		];
 	}
 
 	void FNearestNeighborModelDetails::CustomizeSectionsCategory(IDetailLayoutBuilder& DetailBuilder)
diff --git a/Engine/Plugins/Animation/MLDeformer/NearestNeighborModel/Source/NearestNeighborModelEditor/Private/NearestNeighborModelSectionCustomization.cpp b/Engine/Plugins/Animation/MLDeformer/NearestNeighborModel/Source/NearestNeighborModelEditor/Private/NearestNeighborModelSectionCustomization.cpp
--- a/Engine/Plugins/Animation/MLDeformer/NearestNeighborModel/Source/NearestNeighborModelEditor/Private/NearestNeighborModelSectionCustomization.cpp
+++ b/Engine/Plugins/Animation/MLDeformer/NearestNeighborModel/Source/NearestNeighborModelEditor/Private/NearestNeighborModelSectionCustomization.cpp
@@ -35,19 +35,6 @@ namespace UE::NearestNeighborModel
 			return static_cast<FNearestNeighborEditorModel*>(EditorModule.GetModelRegistry().GetEditorModel(const_cast<UNearestNeighborModel*>(Model)));
 		}
 
-		TArray<float> GetVertexAttributeValues(const TArray<int32>& VertexMap, const TArray<float>& VertexWeights, int32 NumVertices)
-		{
-			check(VertexMap.Num() == VertexWeights.Num());
-			check(VertexMap.Num() <= NumVertices);
-			TArray<float> VertexAttribute;
-			VertexAttribute.SetNumZeroed(NumVertices);
-			for (int32 VertexIndex = 0; VertexIndex < VertexMap.Num(); ++VertexIndex)
-			{
-				VertexAttribute[VertexMap[VertexIndex]] = VertexWeights[VertexIndex];
-			}
-			return VertexAttribute;
-		}
-
 		void CreateVertexAttributes(USkeletalMesh& SkeletalMesh, const FString& AttributeName, const TArray<int32>& VertexMap, const TArray<float>& VertexWeights)
 		{
 			constexpr int32 LODIndex = 0;
@@ -66,9 +53,10 @@ namespace UE::NearestNeighborModel
 			}
 			TVertexAttributesRef<float> AttributeRef = MeshDescription->VertexAttributes().GetAttributesRef<float>(*AttributeName);
 
-			for (const int32 VertexIndex: VertexMap)
+			check(VertexMap.Num() == VertexWeights.Num());
+			for (int32 Index = 0; Index < VertexMap.Num(); ++Index)
 			{
-				AttributeRef.Set(VertexIndex, VertexWeights[VertexIndex]);
+				AttributeRef.Set(VertexMap[Index], VertexWeights[Index]);
 			}
 			
 			SkeletalMesh.CommitMeshDescription(LODIndex);
@@ -300,7 +288,7 @@ namespace UE::NearestNeighborModel
 	{
 		DetailBuilder.HideCategory("Section Private");
 		IDetailCategoryBuilder& SectionBuilder = DetailBuilder.EditCategory("Section", LOCTEXT("SectionCategory", "Section"));
-		SectionBuilder.AddProperty(UNearestNeighborModelSection::GetNumPCACoeffsPropertyName());
+		SectionBuilder.AddProperty(UNearestNeighborModelSection::GetNumBasisPropertyName());
 
 		// Get the selected objects
 		const TArray<TWeakObjectPtr<UObject>>& SelectedObjects = DetailBuilder.GetSelectedObjects();
@@ -356,32 +344,9 @@ namespace UE::NearestNeighborModel
 				? EVisibility::Visible : EVisibility::Collapsed;
 			})));
 
-		const UNearestNeighborModel* const Model = Section->GetModel();
-		if (!Model)
-		{
-			return;
-		}
-		const FName PosesName = UNearestNeighborModelSection::GetNeighborPosesPropertyName();
-		TSharedPtr<IPropertyHandle> NeighborPosesHandle = DetailBuilder.GetProperty(PosesName);
-		SectionBuilder.AddProperty(PosesName)
-			.IsEnabled(TAttribute<bool>::Create(TAttribute<bool>::FGetter::CreateLambda([Model]()
-				{
-					return Model->IsReadyForTraining();
-				})));
-		const FName MeshesName = UNearestNeighborModelSection::GetNeighborMeshesPropertyName();
-		TSharedPtr<IPropertyHandle> NeighborMeshesHandle = DetailBuilder.GetProperty(MeshesName);
-		SectionBuilder.AddProperty(MeshesName)
-			.IsEnabled(TAttribute<bool>::Create(TAttribute<bool>::FGetter::CreateLambda([Model]()
-				{
-					return Model->IsReadyForTraining();
-				})));
-		const FName FramesName = UNearestNeighborModelSection::GetExcludedFramesPropertyName();
-		TSharedPtr<IPropertyHandle> ExcludedFramesHandle = DetailBuilder.GetProperty(FramesName);
-		SectionBuilder.AddProperty(FramesName)
-			.IsEnabled(TAttribute<bool>::Create(TAttribute<bool>::FGetter::CreateLambda([Model]()
-				{
-					return Model->IsReadyForTraining();
-				})));
+		SectionBuilder.AddProperty(UNearestNeighborModelSection::GetNeighborPosesPropertyName());
+		SectionBuilder.AddProperty(UNearestNeighborModelSection::GetNeighborMeshesPropertyName());
+		SectionBuilder.AddProperty(UNearestNeighborModelSection::GetExcludedFramesPropertyName());
 	}
 }	// namespace UE::NearestNeighborModel
 
diff --git a/Engine/Plugins/Animation/MLDeformer/NearestNeighborModel/Source/NearestNeighborModelEditor/Public/NearestNeighborEditorModel.h b/Engine/Plugins/Animation/MLDeformer/NearestNeighborModel/Source/NearestNeighborModelEditor/Public/NearestNeighborEditorModel.h
--- a/Engine/Plugins/Animation/MLDeformer/NearestNeighborModel/Source/NearestNeighborModelEditor/Public/NearestNeighborEditorModel.h
+++ b/Engine/Plugins/Animation/MLDeformer/NearestNeighborModel/Source/NearestNeighborModelEditor/Public/NearestNeighborEditorModel.h
@@ -55,6 +55,7 @@ namespace UE::NearestNeighborModel
 		// ~END UMLDeformerMorphModelEditorModel overrides.
 
 		void OnUpdateClicked();
+		void ClearReferences();
 
 		FVertexMapSelector* GetVertexMapSelector() const;
 		FVertVizSelector* GetVertVizSelector() const;
