NearestNeighborModel:
- add MustIncludeFrames to KMeans tool
- fix a crash when UseRBF and neighbors is empty
- fix temporal smoothing at the first few frames
- Expose ModelInstance to blueprint and add Reset()
- add an option to add vertex weights from External txt file

diff --git a/Engine/Plugins/Animation/MLDeformer/MLDeformerFramework/Source/MLDeformerFramework/Public/MLDeformerComponent.h b/Engine/Plugins/Animation/MLDeformer/MLDeformerFramework/Source/MLDeformerFramework/Public/MLDeformerComponent.h
--- a/Engine/Plugins/Animation/MLDeformer/MLDeformerFramework/Source/MLDeformerFramework/Public/MLDeformerComponent.h
+++ b/Engine/Plugins/Animation/MLDeformer/MLDeformerFramework/Source/MLDeformerFramework/Public/MLDeformerComponent.h
@@ -121,6 +121,7 @@ public:
 	 * The instance is responsible for running inference and feeding the neural network with inputs.
 	 * @return A pointer to the model instance object.
 	 */
+	UFUNCTION(BlueprintCallable, Category = "MLDeformer")
 	UMLDeformerModelInstance* GetModelInstance() const			{ return ModelInstance; }
 
 	/**
diff --git a/Engine/Plugins/Animation/MLDeformer/NearestNeighborModel/Source/NearestNeighborModel/Private/NearestNeighborModel.cpp b/Engine/Plugins/Animation/MLDeformer/NearestNeighborModel/Source/NearestNeighborModel/Private/NearestNeighborModel.cpp
--- a/Engine/Plugins/Animation/MLDeformer/NearestNeighborModel/Source/NearestNeighborModel/Private/NearestNeighborModel.cpp
+++ b/Engine/Plugins/Animation/MLDeformer/NearestNeighborModel/Source/NearestNeighborModel/Private/NearestNeighborModel.cpp
@@ -261,6 +261,38 @@ namespace UE::NearestNeighborModel
 			return (ChangeType & (EPropertyChangeType::ArrayAdd | EPropertyChangeType::ArrayRemove | EPropertyChangeType::ArrayClear | EPropertyChangeType::ArrayMove)) != 0;
 		}
 
+		TOptional<TArray<float>> ReadTxt(const FString& FilePath)
+		{
+			TOptional<TArray<float>> Empty;
+
+			TArray<FString> Lines;
+			if (FFileHelper::LoadFileToStringArray(Lines, *FilePath))
+			{
+				TArray<float> Result;
+				Result.Reserve(Lines.Num());
+				for (const FString& Line : Lines)
+				{
+					float Value;
+					if (Line.TrimStartAndEnd().IsNumeric())
+					{
+						Value = FCString::Atof(*Line);
+						Result.Add(Value);
+					}
+					else
+					{
+						UE_LOG(LogNearestNeighborModel, Error, TEXT("Invalid value in file: %s"), *Line);
+						return Empty;
+					}
+				}
+				return Result;
+			}
+			else
+			{
+				UE_LOG(LogNearestNeighborModel, Error, TEXT("Failed to read file: %s"), *FilePath);
+				return Empty;
+			}
+		}
+
 #if WITH_EDITORONLY_DATA
 		TArray<FInt32Range> GetMeshVertRanges(const USkeletalMesh& SkelMesh)
 		{
@@ -641,6 +673,18 @@ void UNearestNeighborModelSection::SetBoneNames(const TArray<FName>& InBoneNames
 	UpdateVertexWeights();
 }
 
+FString UNearestNeighborModelSection::GetExternalTxtFile() const
+{
+	return ExternalTxtFile;
+}
+
+void UNearestNeighborModelSection::SetExternalTxtFile(const FString& InFile)
+{
+	ExternalTxtFile = InFile;
+	InvalidateTraining();
+	UpdateVertexWeights();
+}
+
 void UNearestNeighborModelSection::PostEditChangeProperty(FPropertyChangedEvent& PropertyChangedEvent)
 {
 	const FProperty* Property = PropertyChangedEvent.Property;
@@ -766,6 +810,14 @@ UE::NearestNeighborModel::EOpFlag UNearestNeighborModelSection::UpdateVertexWeig
 	{
 		Result = UpdateVertexWeightsVertexAttributes();
 	}
+	else if (WeightMapCreationMethod == ENearestNeighborModelSectionWeightMapCreationMethod::ExternalTxt)
+	{
+		Result = UpdateVertexWeightsExternalTxt();
+	}
+	else
+	{
+		UE_LOG(LogNearestNeighborModel, Error, TEXT("Invalid WeightMapCreationMethod."));
+	}
 	using namespace UE::NearestNeighborModel;
 	if (OpFlag::HasError(Result))
 	{
@@ -987,6 +1039,58 @@ UE::NearestNeighborModel::EOpFlag UNearestNeighborModelSection::UpdateVertexWeig
 	return EOpFlag::Success;
 }
 
+UE::NearestNeighborModel::EOpFlag UNearestNeighborModelSection::UpdateVertexWeightsExternalTxt()
+{
+	using namespace UE::NearestNeighborModel;
+	if (!GetModel())
+	{
+		UE_LOG(LogNearestNeighborModel, Error, TEXT("NearestNeighborModel is invalid."));
+		return EOpFlag::Error;
+	}
+	const USkeletalMesh* const SkeletalMesh = GetModel()->GetSkeletalMesh();
+	if (!SkeletalMesh)
+	{
+		UE_LOG(LogNearestNeighborModel, Error, TEXT("SkeletalMesh is None"));
+		return EOpFlag::Error;
+	}
+	if (!GetModel())
+	{
+		UE_LOG(LogNearestNeighborModel, Error, TEXT("NearestNeighborModel is invalid."));
+		return EOpFlag::Error;
+	}
+
+	const int32 NumBaseMeshVerts = Model->GetNumBaseMeshVerts();
+	if (NumBaseMeshVerts <= 0)
+	{
+		UE_LOG(LogNearestNeighborModel, Error, TEXT("SkeletalMesh has no vertices."));
+		return EOpFlag::Error;
+	}
+
+	using Private::ReadTxt;
+	TOptional<TArray<float>> WeightsOpt = ReadTxt(ExternalTxtFile);
+	if (!WeightsOpt.IsSet())
+	{
+		return EOpFlag::Error;
+	}
+	const TArray<float>& Weights = WeightsOpt.GetValue();
+	if (Weights.Num() != NumBaseMeshVerts)
+	{
+		UE_LOG(LogNearestNeighborModel, Error, TEXT("Number of weights %d does not match number of vertices %d."), Weights.Num(), NumBaseMeshVerts);
+		return EOpFlag::Error;
+	}
+	VertexMap.Reset();
+	VertexWeights.Reset();
+	for (int32 Index = 0; Index < Weights.Num(); ++Index)
+	{
+		if (Weights[Index] > 0)
+		{
+			VertexMap.Add(Index);
+			VertexWeights.Add(Weights[Index]);
+		}
+	}
+	return EOpFlag::Success;
+}
+
 UE::NearestNeighborModel::EOpFlag UNearestNeighborModelSection::NormalizeVertexWeights()
 {
 	if (!Model)
diff --git a/Engine/Plugins/Animation/MLDeformer/NearestNeighborModel/Source/NearestNeighborModel/Private/NearestNeighborModelInstance.cpp b/Engine/Plugins/Animation/MLDeformer/NearestNeighborModel/Source/NearestNeighborModel/Private/NearestNeighborModelInstance.cpp
--- a/Engine/Plugins/Animation/MLDeformer/NearestNeighborModel/Source/NearestNeighborModel/Private/NearestNeighborModelInstance.cpp
+++ b/Engine/Plugins/Animation/MLDeformer/NearestNeighborModel/Source/NearestNeighborModel/Private/NearestNeighborModelInstance.cpp
@@ -147,6 +147,11 @@ int64 UNearestNeighborModelInstance::SetBoneTransforms(float* OutputBuffer, int6
 	return EndIndex;
 }
 
+void UNearestNeighborModelInstance::Reset()
+{
+	bNeedsReset = true;
+}
+
 FString UNearestNeighborModelInstance::CheckCompatibility(USkeletalMeshComponent* InSkelMeshComponent, bool LogIssues)
 {
 	ErrorText = FString();
@@ -208,7 +213,6 @@ void UNearestNeighborModelInstance::InitInstanceData(int32 NumMorphWeights)
 		}
 		NumMorphWeights = 1 + NearestNeighborModel->GetTotalNumBasis() + NearestNeighborModel->GetTotalNumNeighbors();
 	}
-	PreviousWeights.SetNumZeroed(NumMorphWeights);
 	DistanceBuffer.SetNumZeroed(NumMorphWeights);
 }
 
@@ -279,14 +283,22 @@ namespace UE::NearestNeighborModel::Private
 		}
 	}
 
-	void UpdateWeightWithDecay(float& OutDecayedWeight, float& InOutPrevWeight, float NewWeight, float DecayCoeff)
+	void UpdateWeightWithDecay(float& OutDecayedWeight, float& InOutPrevWeight, float NewWeight, float DecayCoeff, bool bReset)
 	{
-		const float NewW = (1.0f - DecayCoeff) * NewWeight + DecayCoeff * InOutPrevWeight;
-		OutDecayedWeight = NewW;
-		InOutPrevWeight = NewW;
+		if (bReset)
+		{
+			OutDecayedWeight = NewWeight;
+			InOutPrevWeight = NewWeight;
+		}
+		else
+		{
+			const float NewW = (1.0f - DecayCoeff) * NewWeight + DecayCoeff * InOutPrevWeight;
+			OutDecayedWeight = NewW;
+			InOutPrevWeight = NewW;
+		}
 	}
 
-	void UpdateNearestNeighborWeights(TArrayView<float> OutMorphWeights, TArrayView<float> InOutPrevWeights, TConstArrayView<float> SquaredDistances, float OffsetWeight, float DecayCoeff)
+	void UpdateNearestNeighborWeights(TArrayView<float> OutMorphWeights, TArrayView<float> InOutPrevWeights, TConstArrayView<float> SquaredDistances, float OffsetWeight, float DecayCoeff, bool bReset)
 	{
 		const int32 NumNeighbors = OutMorphWeights.Num();
 		check(InOutPrevWeights.Num() == NumNeighbors);
@@ -302,12 +314,12 @@ namespace UE::NearestNeighborModel::Private
 		for (int32 Index = 0; Index < NumNeighbors; ++Index)
 		{
 			const float W = Index == MinIndex ? OffsetWeight : 0.0f;
-			UpdateWeightWithDecay(OutMorphWeights[Index], InOutPrevWeights[Index], W, DecayCoeff);
+			UpdateWeightWithDecay(OutMorphWeights[Index], InOutPrevWeights[Index], W, DecayCoeff, bReset);
 		}
 	}
 
 	// SquaredDistances buffer will be overwritten after this function.
-	void UpdateRBFWeights(TArrayView<float> OutMorphWeights, TArrayView<float> InOutPrevWeights, TArrayView<float> SquaredDistances, float Sigma, float OffsetWeight, float DecayCoeff)
+	void UpdateRBFWeights(TArrayView<float> OutMorphWeights, TArrayView<float> InOutPrevWeights, TArrayView<float> SquaredDistances, float Sigma, float OffsetWeight, float DecayCoeff, bool bReset)
 	{
 		const int32 NumNeighbors = OutMorphWeights.Num();
 		check(InOutPrevWeights.Num() == NumNeighbors);
@@ -338,7 +350,7 @@ namespace UE::NearestNeighborModel::Private
 		for (int32 Index = 0; Index < NumNeighbors; ++Index)
 		{
 			const float W = Weights[Index] / SumWeights * OffsetWeight;
-			UpdateWeightWithDecay(OutMorphWeights[Index], InOutPrevWeights[Index], W, DecayCoeff);
+			UpdateWeightWithDecay(OutMorphWeights[Index], InOutPrevWeights[Index], W, DecayCoeff, bReset);
 		}
 	}
 }
@@ -381,6 +393,11 @@ void UNearestNeighborModelInstance::RunNearestNeighborModel(float DeltaTime, flo
 	if (NumMorphTargets == NumNetworkWeights + 1 + TotalNumNeighbors)
 	{
 		InitInstanceData(NumMorphTargets);
+		bNeedsReset = PreviousWeights.Num() != NumMorphTargets;
+		if (bNeedsReset)
+		{
+			PreviousWeights.SetNumZeroed(NumMorphTargets);
+		}
 
 		const float OffsetWeight = NearestNeighborModel->GetNearestNeighborOffsetWeight();
 		WeightData->Weights[0] = ModelWeight;
@@ -390,7 +407,7 @@ void UNearestNeighborModelInstance::RunNearestNeighborModel(float DeltaTime, flo
 		{
 			const float W = OutputView[MorphIndex] * ModelWeight;
 			using UE::NearestNeighborModel::Private::UpdateWeightWithDecay;
-			UpdateWeightWithDecay(WeightData->Weights[MorphIndex + 1], PreviousWeights[MorphIndex + 1], W, 0.0f);
+			WeightData->Weights[MorphIndex + 1] = W;
 		}
 
 		int32 NeighborOffset = NumNetworkWeights + 1;
@@ -420,10 +437,18 @@ void UNearestNeighborModelInstance::RunNearestNeighborModel(float DeltaTime, flo
 			TConstArrayView<float> Coeffs(OutputView.GetData() + CoeffStart, NumCoeffs);
 
 			const int32 NumNeighbors = Section.GetRuntimeNumNeighbors();
+			if (NumNeighbors <= 0)
+			{
+				continue;
+			}
 
 			TArrayView<float> SquaredDistances(DistanceBuffer.GetData() + NeighborOffset, NumNeighbors);
 			using UE::NearestNeighborModel::Private::ComputeNeighborDistances;
 			ComputeNeighborDistances(SquaredDistances, Section, Coeffs);
+#if WITH_EDITOR
+			const int32 NeighborId = NumNeighbors > 0 ? (Algo::MinElement(SquaredDistances) - SquaredDistances.GetData()) : INDEX_NONE;
+			NearestNeighborIds[SectionIndex] = NeighborId;
+#endif // WITH_EDITOR
 
 			TArrayView<float> SectionMorphWeights(WeightData->Weights.GetData() + NeighborOffset, NumNeighbors);
 			TArrayView<float> SectionPreviousWeights(PreviousWeights.GetData() + NeighborOffset, NumNeighbors);
@@ -431,21 +456,21 @@ void UNearestNeighborModelInstance::RunNearestNeighborModel(float DeltaTime, flo
 			{
 				const float Sigma = NearestNeighborModel->GetRBFSigma();
 				using UE::NearestNeighborModel::Private::UpdateRBFWeights;
-				UpdateRBFWeights(SectionMorphWeights, SectionPreviousWeights, SquaredDistances, Sigma, OffsetWeight * ModelWeight, DecayCoeff);
+				UpdateRBFWeights(SectionMorphWeights, SectionPreviousWeights, SquaredDistances, Sigma, OffsetWeight * ModelWeight, DecayCoeff, bNeedsReset);
 			}
 			else
 			{
 				using UE::NearestNeighborModel::Private::UpdateNearestNeighborWeights;
-				UpdateNearestNeighborWeights(SectionMorphWeights, SectionPreviousWeights, SquaredDistances, OffsetWeight * ModelWeight, DecayCoeff);
+				UpdateNearestNeighborWeights(SectionMorphWeights, SectionPreviousWeights, SquaredDistances, OffsetWeight * ModelWeight, DecayCoeff, bNeedsReset);
 			}
 
-#if WITH_EDITOR
-			const int32 NeighborId = NumNeighbors > 0 ? (Algo::MinElement(SquaredDistances) - SquaredDistances.GetData()) : INDEX_NONE;
-			NearestNeighborIds[SectionIndex] = NeighborId;
-#endif // WITH_EDITOR
-
 			NeighborOffset += NumNeighbors;
 		}
+
+		if (bNeedsReset)
+		{
+			bNeedsReset = false;
+		}
 	}
 }
 
diff --git a/Engine/Plugins/Animation/MLDeformer/NearestNeighborModel/Source/NearestNeighborModel/Public/NearestNeighborModel.h b/Engine/Plugins/Animation/MLDeformer/NearestNeighborModel/Source/NearestNeighborModel/Public/NearestNeighborModel.h
--- a/Engine/Plugins/Animation/MLDeformer/NearestNeighborModel/Source/NearestNeighborModel/Public/NearestNeighborModel.h
+++ b/Engine/Plugins/Animation/MLDeformer/NearestNeighborModel/Source/NearestNeighborModel/Public/NearestNeighborModel.h
@@ -80,6 +80,8 @@ enum class ENearestNeighborModelSectionWeightMapCreationMethod : uint8
 	SelectedBones,
 	/** Use weights from a vertex attribute. */
 	VertexAttributes,
+	/** Using an external .txt file */
+	ExternalTxt
 };
 #endif
 
@@ -171,6 +173,8 @@ public:
 	FString GetBoneNamesString() const;
 	const TArray<FName>& GetBoneNames() const;
 	void SetBoneNames(const TArray<FName>& InBoneNames);
+	FString GetExternalTxtFile() const;
+	void SetExternalTxtFile(const FString& InFile);
 
 
 	// Do not call this function directly. Call UNearestNeighborModel::NormalizeVertexWeights() instead.
@@ -187,6 +191,7 @@ public:
 	static FName GetExcludedFramesPropertyName() { return GET_MEMBER_NAME_CHECKED(UNearestNeighborModelSection, ExcludedFrames); }
 	static FName GetWeightMapCreationMethodPropertyName() { return GET_MEMBER_NAME_CHECKED(UNearestNeighborModelSection, WeightMapCreationMethod); }
 	static FName GetAttributeNamePropertyName() { return GET_MEMBER_NAME_CHECKED(UNearestNeighborModelSection, AttributeName); }
+	static FName GetExternalTxtFilePropertyName() { return GET_MEMBER_NAME_CHECKED(UNearestNeighborModelSection, ExternalTxtFile); }
 #endif
 #if WITH_EDITORONLY_DATA
 PRAGMA_DISABLE_DEPRECATION_WARNINGS
@@ -221,6 +226,10 @@ protected:
 	UPROPERTY(EditAnywhere, Category = "Section", meta = (DisplayName = "Vertex Indices"))
 	FString VertexMapString;
 
+	/** The path to the txt file containing vertex weights. The number of lines equals to the number of vertices in skeletal mesh with each line being a float for vertex weight. */
+	UPROPERTY(VisibleAnywhere, Category = "Section")
+	FString ExternalTxtFile;
+
 	/** Frames to be excluded from the nearest neighbor ROM */
 	UPROPERTY(EditAnywhere, Category = "Section")
 	TArray<int32> ExcludedFrames;
@@ -305,6 +314,7 @@ private:
 	EOpFlag UpdateVertexWeightsFromText();
 	EOpFlag UpdateVertexWeightsSelectedBones();
 	EOpFlag UpdateVertexWeightsVertexAttributes();
+	EOpFlag UpdateVertexWeightsExternalTxt();
 
 	UFUNCTION()
 	TArray<FName> GetVertexAttributeNames() const;
diff --git a/Engine/Plugins/Animation/MLDeformer/NearestNeighborModel/Source/NearestNeighborModel/Public/NearestNeighborModelInstance.h b/Engine/Plugins/Animation/MLDeformer/NearestNeighborModel/Source/NearestNeighborModel/Public/NearestNeighborModelInstance.h
--- a/Engine/Plugins/Animation/MLDeformer/NearestNeighborModel/Source/NearestNeighborModel/Public/NearestNeighborModelInstance.h
+++ b/Engine/Plugins/Animation/MLDeformer/NearestNeighborModel/Source/NearestNeighborModel/Public/NearestNeighborModelInstance.h
@@ -16,7 +16,7 @@ namespace UE::NearestNeighborModel
     class FNearestNeighborEditorModelActor;
 };
 
-UCLASS()
+UCLASS(Blueprintable)
 class NEARESTNEIGHBORMODEL_API UNearestNeighborModelInstance
     : public UMLDeformerMorphModelInstance
 {
@@ -32,6 +32,9 @@ public:
     virtual int64 SetBoneTransforms(float* OutputBuffer, int64 OutputBufferSize, int64 StartIndex) override;
     // ~END UMLDeformerModelInstance overrides
 
+    UFUNCTION(BlueprintCallable, Category = "NearestNeighborModel")
+    void Reset();
+
     friend class UNearestNeighborModelInputInfo;
     friend class UNearestNeighborModelInstance;
     friend class UNearestNeighborOptimizedNetworkInstance;
@@ -65,6 +68,7 @@ private:
 
     TArray<float> PreviousWeights;
     TArray<float> DistanceBuffer;
+    bool bNeedsReset = true;
 
     UPROPERTY()
     TObjectPtr<UNearestNeighborOptimizedNetworkInstance> OptimizedNetworkInstance;
diff --git a/Engine/Plugins/Animation/MLDeformer/NearestNeighborModel/Source/NearestNeighborModelEditor/Private/NearestNeighborEditorModel.cpp b/Engine/Plugins/Animation/MLDeformer/NearestNeighborModel/Source/NearestNeighborModelEditor/Private/NearestNeighborEditorModel.cpp
--- a/Engine/Plugins/Animation/MLDeformer/NearestNeighborModel/Source/NearestNeighborModelEditor/Private/NearestNeighborEditorModel.cpp
+++ b/Engine/Plugins/Animation/MLDeformer/NearestNeighborModel/Source/NearestNeighborModelEditor/Private/NearestNeighborEditorModel.cpp
@@ -501,10 +501,9 @@ namespace UE::NearestNeighborModel
 		{
 			return UpdateResult;
 		}
-		UpdateIsReadyForTrainingState();
 		if (!IsReadyForTraining())
 		{
-			UE_LOG(LogNearestNeighborModel, Error, TEXT("Model is not ready for training. Please check training data is not empty or reload MLDeformer editor."));
+			UE_LOG(LogNearestNeighborModel, Error, TEXT("Model is not ready for training. Please check if training data is not empty or reload MLDeformer editor."));
 			return EOpFlag::Error;
 		}
 		UpdateResult |= UpdateNearestNeighborData();
@@ -567,6 +566,7 @@ namespace UE::NearestNeighborModel
 		{
 			Editor->GetVizSettingsDetailsView()->ForceRefresh();
 		}
+		UpdateIsReadyForTrainingState();
 	}
 
 	EOpFlag FNearestNeighborEditorModel::CheckNetwork()
diff --git a/Engine/Plugins/Animation/MLDeformer/NearestNeighborModel/Source/NearestNeighborModelEditor/Private/NearestNeighborModelSectionCustomization.cpp b/Engine/Plugins/Animation/MLDeformer/NearestNeighborModel/Source/NearestNeighborModelEditor/Private/NearestNeighborModelSectionCustomization.cpp
--- a/Engine/Plugins/Animation/MLDeformer/NearestNeighborModel/Source/NearestNeighborModelEditor/Private/NearestNeighborModelSectionCustomization.cpp
+++ b/Engine/Plugins/Animation/MLDeformer/NearestNeighborModel/Source/NearestNeighborModelEditor/Private/NearestNeighborModelSectionCustomization.cpp
@@ -2,6 +2,7 @@
 
 #include "NearestNeighborModelSectionCustomization.h"
 
+#include "DesktopPlatformModule.h"
 #include "DetailWidgetRow.h"
 #include "DetailLayoutBuilder.h"
 #include "Engine/SkeletalMesh.h"
@@ -257,6 +258,37 @@ namespace UE::NearestNeighborModel
 				LayoutBuilder.HideProperty(VertexMapHandle);
 			}
 		}
+
+		TOptional<FString> OpenTxtFileDialog(const FString& DefaultPath)
+		{
+			TOptional<FString> Empty;
+			IDesktopPlatform* DesktopPlatform = FDesktopPlatformModule::Get();
+			if (DesktopPlatform == nullptr)
+			{
+				return Empty;
+			}
+
+			TArray<FString> FileTypes;
+			FileTypes.Add("Text Files (*.txt)|*.txt");
+
+			TArray<FString> OutFileNames;
+			const bool bFileDialogOpened = DesktopPlatform->OpenFileDialog(
+				FSlateApplication::Get().FindBestParentWindowHandleForDialogs(nullptr),
+				TEXT("Choose a txt"),
+				DefaultPath,
+				TEXT(""),
+				FileTypes[0],
+				EFileDialogFlags::None,
+				OutFileNames
+			);
+
+			if (bFileDialogOpened && OutFileNames.Num() > 0)
+			{
+				return OutFileNames[0];
+			}
+
+			return Empty;
+		}
 	};
 
 	void SNearestNeighborModelSectionWidget::Construct(const FArguments& InArgs)
@@ -344,6 +376,60 @@ namespace UE::NearestNeighborModel
 				? EVisibility::Visible : EVisibility::Collapsed;
 			})));
 
+		
+		IDetailPropertyRow& ExternalTxtFileRow = SectionBuilder.AddProperty(UNearestNeighborModelSection::GetExternalTxtFilePropertyName())
+			.OverrideResetToDefault(FResetToDefaultOverride::Hide())
+			.Visibility(TAttribute<EVisibility>::Create(TAttribute<EVisibility>::FGetter::CreateLambda([Section]()
+			{
+				return Section->GetWeightMapCreationMethod() == 
+				ENearestNeighborModelSectionWeightMapCreationMethod::ExternalTxt 
+				? EVisibility::Visible : EVisibility::Collapsed;
+			})));
+		
+		ExternalTxtFileRow.CustomWidget()
+		.NameContent()
+		[
+			ExternalTxtFileRow.GetPropertyHandle()->CreatePropertyNameWidget()
+		]
+		.ValueContent()
+		[
+			SNew(SHorizontalBox)
+			+ SHorizontalBox::Slot()
+			.FillWidth(0.8f)
+			.VAlign(VAlign_Center)
+			[
+				ExternalTxtFileRow.GetPropertyHandle()->CreatePropertyValueWidget()
+			]
+			+ SHorizontalBox::Slot()
+			.AutoWidth()
+			.VAlign(VAlign_Center)
+			[
+				SNew(SBox)
+				.Padding(2, 2)
+				.MinDesiredWidth(20.f)
+				[
+					SNew(SButton)
+					.Text(LOCTEXT("ChooseExternalTxtButton", "..."))
+					.OnClicked_Lambda([Section]()
+					{
+						using ::UE::NearestNeighborModel::Private::OpenTxtFileDialog;
+						const FString CurrentPath = Section->GetExternalTxtFile();
+						FString DefaultPath = FPaths::ProjectIntermediateDir();
+						if (CurrentPath.Len() > 0)
+						{
+							DefaultPath = FPaths::GetPath(CurrentPath);
+						}
+						TOptional<FString> PathOpt = OpenTxtFileDialog(DefaultPath);
+						if (PathOpt.IsSet())
+						{
+							Section->SetExternalTxtFile(PathOpt.GetValue());
+						}
+						return FReply::Handled();
+					})
+				]
+			]
+		];
+
 		SectionBuilder.AddProperty(UNearestNeighborModelSection::GetNeighborPosesPropertyName());
 		SectionBuilder.AddProperty(UNearestNeighborModelSection::GetNeighborMeshesPropertyName());
 		SectionBuilder.AddProperty(UNearestNeighborModelSection::GetExcludedFramesPropertyName());
diff --git a/Engine/Plugins/Animation/MLDeformer/NearestNeighborModel/Source/NearestNeighborModelEditor/Private/Tools/NearestNeighborKMeansTool.cpp b/Engine/Plugins/Animation/MLDeformer/NearestNeighborModel/Source/NearestNeighborModelEditor/Private/Tools/NearestNeighborKMeansTool.cpp
--- a/Engine/Plugins/Animation/MLDeformer/NearestNeighborModel/Source/NearestNeighborModelEditor/Private/Tools/NearestNeighborKMeansTool.cpp
+++ b/Engine/Plugins/Animation/MLDeformer/NearestNeighborModel/Source/NearestNeighborModelEditor/Private/Tools/NearestNeighborKMeansTool.cpp
@@ -169,10 +169,20 @@ namespace UE::NearestNeighborModel
 					{
 						if (UNearestNeighborKMeansData* KMeansData = Cast<UNearestNeighborKMeansData>(&Data))
 						{
-							UNearestNeighborTrainingModel* TrainingModel = FHelpers::NewDerivedObject<UNearestNeighborTrainingModel>();
-							TrainingModel->Init(EditorModel.Get());
-							const int32 ResultInt = TrainingModel->KmeansClusterPoses(KMeansData);
-							const EOpFlag Result = ToOpFlag(ResultInt);
+							EOpFlag Result = EOpFlag::Success;
+							if (!EditorModel->IsReadyForTraining())
+							{
+								Result = EOpFlag::Error;
+								UE_LOG(LogNearestNeighborModel, Error, TEXT("Model is not ready for training. Please check if training data is not empty or reload MLDeformer editor."));
+							}
+							else
+							{
+								UNearestNeighborTrainingModel* TrainingModel = FHelpers::NewDerivedObject<UNearestNeighborTrainingModel>();
+								TrainingModel->Init(EditorModel.Get());
+								const int32 ResultInt = TrainingModel->KmeansClusterPoses(KMeansData);
+								Result = ToOpFlag(ResultInt);
+							}
+
 							const FText WindowTitle = LOCTEXT("KmeansWindowTitle", "Extraction Results");
 							if (OpFlag::HasError(Result))
 							{
diff --git a/Engine/Plugins/Animation/MLDeformer/NearestNeighborModel/Source/NearestNeighborModelEditor/Private/Tools/NearestNeighborKMeansTool.h b/Engine/Plugins/Animation/MLDeformer/NearestNeighborModel/Source/NearestNeighborModelEditor/Private/Tools/NearestNeighborKMeansTool.h
--- a/Engine/Plugins/Animation/MLDeformer/NearestNeighborModel/Source/NearestNeighborModelEditor/Private/Tools/NearestNeighborKMeansTool.h
+++ b/Engine/Plugins/Animation/MLDeformer/NearestNeighborModel/Source/NearestNeighborModelEditor/Private/Tools/NearestNeighborKMeansTool.h
@@ -18,6 +18,24 @@ namespace UE::MLDeformer
 	class FMLDeformerEditorModel;
 }
 
+class UNearestNeighborKMeansData;
+
+USTRUCT(BlueprintType)
+struct NEARESTNEIGHBORMODELEDITOR_API FNearestNeighborKMeansInputData
+{
+	GENERATED_BODY()
+public:
+	UPROPERTY(EditAnywhere, BlueprintReadOnly, Category = "Input")
+	TObjectPtr<UAnimSequence> Poses;
+
+	UPROPERTY(EditAnywhere, BlueprintReadOnly, Category = "Input", DisplayName = "(Optional) Cache" )
+	TObjectPtr<UGeometryCache> Cache;
+
+	UPROPERTY(EditAnywhere, BlueprintReadOnly, Category = "Input")
+	TArray<int> MustIncludeFrames;
+};
+
+
 UCLASS(Blueprintable)
 class NEARESTNEIGHBORMODELEDITOR_API UNearestNeighborKMeansData : public UObject
 {
@@ -34,17 +52,12 @@ public:
 	UPROPERTY(EditAnywhere, BlueprintReadOnly, Category = "Input", meta = (ClampMin = 1))
 	int32 NumClusters = 10;
 
-	/** List of input poses (cannot be empty). */
-	UPROPERTY(EditAnywhere, BlueprintReadOnly, Category = "Input")
-	TArray<TObjectPtr<UAnimSequence>> InputPoses;
-
 	/** Whether to extract geometry cache at the same time. */
 	UPROPERTY(EditAnywhere, BlueprintReadOnly, Category = "Input")
 	bool bExtractGeometryCache = false;
 
-	/** List of input geometry caches (need to be the same size as InputPoses). */
-	UPROPERTY(EditAnywhere, BlueprintReadOnly, Category = "Input", meta = (EditCondition = "bExtractGeometryCache"))
-	TArray<TObjectPtr<UGeometryCache>> InputCaches;
+	UPROPERTY(EditAnywhere, BlueprintReadOnly, Category = "Input")
+	TArray<FNearestNeighborKMeansInputData> Inputs;
 
 	/** Extracted poses. */
 	UPROPERTY(EditAnywhere, BlueprintReadOnly, Category = "Output")
