Unshelved from pending changelist '32101130':

NearestNeighborModel: change to beta

diff --git a/Engine/Plugins/Animation/MLDeformer/NearestNeighborModel/NearestNeighborModel.uplugin b/Engine/Plugins/Animation/MLDeformer/NearestNeighborModel/NearestNeighborModel.uplugin
--- a/Engine/Plugins/Animation/MLDeformer/NearestNeighborModel/NearestNeighborModel.uplugin
+++ b/Engine/Plugins/Animation/MLDeformer/NearestNeighborModel/NearestNeighborModel.uplugin
@@ -1,7 +1,7 @@
 {
 	"FileVersion": 3,
 	"Version": 1,
-	"VersionName": "0.2",
+	"VersionName": "0.3",
 	"FriendlyName": "ML Deformer Nearest Neighbor Model",
 	"Description": "Nearest Neighbor Model for the ML Deformer Framework",
 	"Category": "Animation",
@@ -11,8 +11,8 @@
 	"MarketplaceURL": "",
 	"SupportURL": "https://forums.unrealengine.com/",
 	"CanContainContent": true,
-	"IsBetaVersion": false,
-	"IsExperimentalVersion": true,
+	"IsBetaVersion": true,
+	"IsExperimentalVersion": false,
 	"Installed": false,
 	"Modules": [
 		{
diff --git a/Engine/Plugins/Animation/MLDeformer/NearestNeighborModel/Source/NearestNeighborModelEditor/Private/NearestNeighborModelDetails.cpp b/Engine/Plugins/Animation/MLDeformer/NearestNeighborModel/Source/NearestNeighborModelEditor/Private/NearestNeighborModelDetails.cpp
--- a/Engine/Plugins/Animation/MLDeformer/NearestNeighborModel/Source/NearestNeighborModelEditor/Private/NearestNeighborModelDetails.cpp
+++ b/Engine/Plugins/Animation/MLDeformer/NearestNeighborModel/Source/NearestNeighborModelEditor/Private/NearestNeighborModelDetails.cpp
@@ -44,19 +44,6 @@ namespace UE::NearestNeighborModel
 		NearestNeighborCategoryBuilder = &DetailLayoutBuilder->EditCategory("Nearest Neighbor Setting", FText::GetEmpty(), ECategoryPriority::Important);
 		SectionsCategoryBuilder = &DetailLayoutBuilder->EditCategory("Sections", FText::GetEmpty(), ECategoryPriority::Important);
 		StatusCategoryBuilder = &DetailLayoutBuilder->EditCategory("Status", FText::GetEmpty(), ECategoryPriority::Important);
-
-		// Add warning in CreateCategories so that the warning appears at the top of the details panel.
-		FDetailWidgetRow& NearestNeighborWarningRow = BaseMeshCategoryBuilder->AddCustomRow(FText::FromString("NearestNeighborWarning"))
-			.WholeRowContent()
-			[
-				SNew(SBox)
-				.Padding(FMargin(0.0f, 4.0f))
-				[
-					SNew(SWarningOrErrorBox)
-					.MessageStyle(EMessageStyle::Warning)
-					.Message(LOCTEXT("NearestNeighborWarning", "Nearest neighbor model is still experimental and the details here are subject to change."))
-				]
-			];
 	}
 
 	UNearestNeighborModel* FNearestNeighborModelDetails::GetCastModel() const
