NearestNeighborModel: remove NumEpochs and change to NumIterations

diff --git a/Engine/Plugins/Animation/MLDeformer/NearestNeighborModel/Source/NearestNeighborModel/Private/NearestNeighborModel.cpp b/Engine/Plugins/Animation/MLDeformer/NearestNeighborModel/Source/NearestNeighborModel/Private/NearestNeighborModel.cpp
--- a/Engine/Plugins/Animation/MLDeformer/NearestNeighborModel/Source/NearestNeighborModel/Private/NearestNeighborModel.cpp
+++ b/Engine/Plugins/Animation/MLDeformer/NearestNeighborModel/Source/NearestNeighborModel/Private/NearestNeighborModel.cpp
@@ -54,6 +54,7 @@ namespace UE::NearestNeighborModel
 		{
 			BeforeCustomVersionWasAdded = 0,
 			AddTrainedBasis = 1,
+			DeprecateNumEpochs = 2,
 	
 			VersionPlusOne,
 			LatestVersion = VersionPlusOne - 1
@@ -1182,7 +1183,7 @@ void UNearestNeighborModel::GetAssetRegistryTags(FAssetRegistryTagsContext Conte
 	Super::GetAssetRegistryTags(Context);
 
 	#if WITH_EDITORONLY_DATA
-		Context.AddTag(FAssetRegistryTag("MLDeformer.NearestNeighborModel.NumEpochs", FString::FromInt(NumEpochs), FAssetRegistryTag::TT_Numerical));
+		Context.AddTag(FAssetRegistryTag("MLDeformer.NearestNeighborModel.NumIterations", FString::FromInt(NumIterations), FAssetRegistryTag::TT_Numerical));
 		Context.AddTag(FAssetRegistryTag("MLDeformer.NearestNeighborModel.BatchSize", FString::FromInt(BatchSize), FAssetRegistryTag::TT_Numerical));
 		Context.AddTag(FAssetRegistryTag("MLDeformer.NearestNeighborModel.NumHiddenLayers", FString::FromInt(HiddenLayerDims.Num()), FAssetRegistryTag::TT_Numerical));
 		Context.AddTag(FAssetRegistryTag("MLDeformer.NearestNeighborModel.LearningRate", FString::Printf(TEXT("%f"), LearningRate), FAssetRegistryTag::TT_Numerical));
@@ -1211,7 +1212,7 @@ void UNearestNeighborModel::PostEditChangeProperty(FPropertyChangedEvent& Proper
 		Property->GetFName() == UMLDeformerModel::GetCurveIncludeListPropertyName() ||
 		Property->GetFName() == UMLDeformerModel::GetMaxTrainingFramesPropertyName() ||
 		Property->GetFName() == UNearestNeighborModel::GetHiddenLayerDimsPropertyName() ||
-		Property->GetFName() == UNearestNeighborModel::GetNumEpochsPropertyName() ||
+		Property->GetFName() == UNearestNeighborModel::GetNumIterationsPropertyName() ||
 		Property->GetFName() == UNearestNeighborModel::GetBatchSizePropertyName() ||
 		Property->GetFName() == UNearestNeighborModel::GetLearningRatePropertyName() ||
 		Property->GetFName() == UNearestNeighborModel::GetEarlyStopEpochsPropertyName())
@@ -1394,6 +1395,11 @@ void UNearestNeighborModel::PostLoad()
 	{
 		bUsePCA = true;
 	}
+
+	if (IsBeforeDeprecateNumEpochs())
+	{
+		NumIterations = 20000;
+	}
 #endif
 }
 
@@ -1660,6 +1666,12 @@ bool UNearestNeighborModel::IsBeforeTrainedBasisAdded() const
 	return Version < FNearestNeighborModelCustomVersion::AddTrainedBasis;
 }
 
+bool UNearestNeighborModel::IsBeforeDeprecateNumEpochs() const
+{
+	using UE::NearestNeighborModel::FNearestNeighborModelCustomVersion;
+	return Version < FNearestNeighborModelCustomVersion::DeprecateNumEpochs;
+}
+
 UE::NearestNeighborModel::EOpFlag UNearestNeighborModel::CheckHiddenLayerDims()
 {
 	check(InputDim > 0 && OutputDim > 0);
diff --git a/Engine/Plugins/Animation/MLDeformer/NearestNeighborModel/Source/NearestNeighborModel/Public/NearestNeighborModel.h b/Engine/Plugins/Animation/MLDeformer/NearestNeighborModel/Source/NearestNeighborModel/Public/NearestNeighborModel.h
--- a/Engine/Plugins/Animation/MLDeformer/NearestNeighborModel/Source/NearestNeighborModel/Public/NearestNeighborModel.h
+++ b/Engine/Plugins/Animation/MLDeformer/NearestNeighborModel/Source/NearestNeighborModel/Public/NearestNeighborModel.h
@@ -203,11 +203,11 @@ PRAGMA_ENABLE_DEPRECATION_WARNINGS
 protected:
 #if WITH_EDITORONLY_DATA
 	/** Poses of the nearest neighbor ROM. */
-	UPROPERTY(EditAnywhere, Category = "Section")
+	UPROPERTY(EditAnywhere, Category = "Section", meta = (DisplayName = "Neighbor Sequence"))
 	TObjectPtr<UAnimSequence> NeighborPoses;
 
 	/** Geometry cache of the nearest neighbor ROM. */
-	UPROPERTY(EditAnywhere, Category = "Section")
+	UPROPERTY(EditAnywhere, Category = "Section", meta = (DisplayName = "Neighbor Geom Cache"))
 	TObjectPtr<UGeometryCache> NeighborMeshes;
 
 	/** Method to create weight map for this section. */
@@ -417,8 +417,12 @@ public:
 	UFUNCTION(BlueprintPure, Category = "Nearest Neighbor Model")
 	int32 GetOutputDim() const { return OutputDim; }
 	
+	UE_DEPRECATED(5.5, "GetNumEpochs() is deprecated.")
+	UFUNCTION(meta = (DeprecatedFunction, DeprecationMessage = "GetNumEpochs has been deprecated. Convert from GetNumIterations instead."))
+	int32 GetNumEpochs() const { return 0; }
+
 	UFUNCTION(BlueprintPure, Category = "Nearest Neighbor Model")
-	int32 GetNumEpochs() const { return NumEpochs; }
+	int32 GetNumIterations() const { return NumIterations; }
 	
 	UFUNCTION(BlueprintPure, Category = "Nearest Neighbor Model")
 	int32 GetBatchSize() const { return BatchSize; }
@@ -486,12 +490,13 @@ public:
 
 	bool IsBeforeCustomVersionWasAdded() const;
 	bool IsBeforeTrainedBasisAdded() const;
+	bool IsBeforeDeprecateNumEpochs() const;
 	const TArray<float>& GetVertexWeightSum() const;
 
 	static FName GetInputDimPropertyName() { return GET_MEMBER_NAME_CHECKED(UNearestNeighborModel, InputDim); }
 	static FName GetHiddenLayerDimsPropertyName() { return GET_MEMBER_NAME_CHECKED(UNearestNeighborModel, HiddenLayerDims); }
 	static FName GetOutputDimPropertyName() { return GET_MEMBER_NAME_CHECKED(UNearestNeighborModel, OutputDim); }
-	static FName GetNumEpochsPropertyName() { return GET_MEMBER_NAME_CHECKED(UNearestNeighborModel, NumEpochs); }
+	static FName GetNumIterationsPropertyName() { return GET_MEMBER_NAME_CHECKED(UNearestNeighborModel, NumIterations); }
 	static FName GetBatchSizePropertyName() { return GET_MEMBER_NAME_CHECKED(UNearestNeighborModel, BatchSize); }
 	static FName GetLearningRatePropertyName() { return GET_MEMBER_NAME_CHECKED(UNearestNeighborModel, LearningRate); }
 	static FName GetEarlyStopEpochsPropertyName() { return GET_MEMBER_NAME_CHECKED(UNearestNeighborModel, EarlyStopEpochs); }
@@ -522,8 +527,13 @@ protected:
 	TArray<int32> HiddenLayerDims;
 
 	/** Max number of cycles iterated through the training set. */
-	UPROPERTY(EditAnywhere, BlueprintReadWrite, Category = "Training Settings", AdvancedDisplay, meta = (ClampMin = "1"))
-	int32 NumEpochs = 2500;
+	UE_DEPRECATED(5.5, "NumEpochs is deprecated. Convert to NumIterations instead.")
+	UPROPERTY(meta = (DeprecatedProperty, DeprecationMessage = "NumEpochs is deprecated. Convert to NumIterations instead.", ClampMin = "1"))
+	int32 NumEpochs_DEPRECATED = 2500;
+
+	/** The number of iterations to train the model for. */
+	UPROPERTY(EditAnywhere, BlueprintReadWrite, Category = "Training Settings", meta = (ClampMin = "1", ClampMax = "1000000"))
+	int32 NumIterations = 5000;
 
 	/** Number of data samples processed together as a group in a single pass. */
 	UPROPERTY(EditAnywhere, BlueprintReadWrite, Category = "Training Settings", AdvancedDisplay, meta = (ClampMin = "1"))
diff --git a/Engine/Plugins/Animation/MLDeformer/NearestNeighborModel/Source/NearestNeighborModelEditor/Private/NearestNeighborModelDetails.cpp b/Engine/Plugins/Animation/MLDeformer/NearestNeighborModel/Source/NearestNeighborModelEditor/Private/NearestNeighborModelDetails.cpp
--- a/Engine/Plugins/Animation/MLDeformer/NearestNeighborModel/Source/NearestNeighborModelEditor/Private/NearestNeighborModelDetails.cpp
+++ b/Engine/Plugins/Animation/MLDeformer/NearestNeighborModel/Source/NearestNeighborModelEditor/Private/NearestNeighborModelDetails.cpp
@@ -192,7 +192,7 @@ namespace UE::NearestNeighborModel
 		TrainingSettingsCategoryBuilder->AddProperty(UNearestNeighborModel::GetInputDimPropertyName());
 		TrainingSettingsCategoryBuilder->AddProperty(UNearestNeighborModel::GetHiddenLayerDimsPropertyName());
 		TrainingSettingsCategoryBuilder->AddProperty(UNearestNeighborModel::GetOutputDimPropertyName());
-		TrainingSettingsCategoryBuilder->AddProperty(UNearestNeighborModel::GetNumEpochsPropertyName());
+		TrainingSettingsCategoryBuilder->AddProperty(UNearestNeighborModel::GetNumIterationsPropertyName());
 		TrainingSettingsCategoryBuilder->AddProperty(UNearestNeighborModel::GetBatchSizePropertyName());
 		TrainingSettingsCategoryBuilder->AddProperty(UNearestNeighborModel::GetLearningRatePropertyName());
 	}
