[MLDeformer] Fix a bug in the Nearest Neighbor model where it doesn't show any deformations after training.

#rnx

diff --git a/Engine/Plugins/Animation/MLDeformer/NearestNeighborModel/Source/NearestNeighborModelEditor/Private/NearestNeighborEditorModel.cpp b/Engine/Plugins/Animation/MLDeformer/NearestNeighborModel/Source/NearestNeighborModelEditor/Private/NearestNeighborEditorModel.cpp
--- a/Engine/Plugins/Animation/MLDeformer/NearestNeighborModel/Source/NearestNeighborModelEditor/Private/NearestNeighborEditorModel.cpp
+++ b/Engine/Plugins/Animation/MLDeformer/NearestNeighborModel/Source/NearestNeighborModelEditor/Private/NearestNeighborEditorModel.cpp
@@ -177,6 +177,7 @@ namespace UE::NearestNeighborModel
 
 	void FNearestNeighborEditorModel::OnPostTraining(ETrainingResult TrainingResult, bool bUsePartiallyTrainedWhenAborted)
 	{
+		UNearestNeighborModel* const NearestNeighborModel = GetCastModel();
 		if (TrainingResult == ETrainingResult::Aborted && !bUsePartiallyTrainedWhenAborted)
 		{
 			GetMorphModel()->SetMorphTargetDeltas(MorphTargetDeltasBackup);
@@ -184,18 +185,20 @@ namespace UE::NearestNeighborModel
 		}
 		else if (TrainingResult == ETrainingResult::Success || (TrainingResult == ETrainingResult::Aborted && bUsePartiallyTrainedWhenAborted))
 		{
-			UNearestNeighborModel* const NearestNeighborModel = GetCastModel();
-			if (!NearestNeighborModel)
-			{
-				return;
-			}
-			NearestNeighborModel->InvalidateInference();
-			if (NearestNeighborModel->DoesUseFileCache())
+			if (NearestNeighborModel)
 			{
-				NearestNeighborModel->UpdateFileCache();
+				NearestNeighborModel->InvalidateInference();
+				if (NearestNeighborModel->DoesUseFileCache())
+				{
+					NearestNeighborModel->UpdateFileCache();
+				}
 			}
 			ResetMorphTargets();
 		}
+		if (NearestNeighborModel)
+		{
+			NearestNeighborModel->UpdateForInference();
+		}
 		FMLDeformerMorphModelEditorModel::OnPostTraining(TrainingResult, bUsePartiallyTrainedWhenAborted);
 	}
 
