Nearest Neighbor Model: Set the correct Input/Output Dim settings when upgrading old assets.

Also set NumBasisPerSection to EditConditionHides since it's not supposed to do anything when using PCA, and fix a bug in UpdatePCACoeffStarts where the array wasn't getting reset.

diff --git a/Engine/Plugins/Animation/MLDeformer/NearestNeighborModel/Source/NearestNeighborModel/Private/NearestNeighborModel.cpp b/Engine/Plugins/Animation/MLDeformer/NearestNeighborModel/Source/NearestNeighborModel/Private/NearestNeighborModel.cpp
--- a/Engine/Plugins/Animation/MLDeformer/NearestNeighborModel/Source/NearestNeighborModel/Private/NearestNeighborModel.cpp
+++ b/Engine/Plugins/Animation/MLDeformer/NearestNeighborModel/Source/NearestNeighborModel/Private/NearestNeighborModel.cpp
@@ -1362,19 +1362,8 @@ void UNearestNeighborModel::PostLoad()
 	Super::PostLoad();
 
 #if WITH_EDITORONLY_DATA
-	UpdateNetworkInputDim();
-	UpdateNetworkOutputDim();
-
-	UpdateFileCache();
-
-	for (FSection* Section : Sections)
-	{
-		if (Section)
-		{
-			Section->SetModel(this);
-		}
-	}
 
+	// Apply any upgrades
 	if (IsBeforeCustomVersionWasAdded())
 	{
 	PRAGMA_DISABLE_DEPRECATION_WARNINGS 
@@ -1400,6 +1389,20 @@ void UNearestNeighborModel::PostLoad()
 	{
 		NumIterations = 20000;
 	}
+
+	// Update data set at runtime.
+	UpdateNetworkInputDim();
+	UpdateNetworkOutputDim();
+
+	UpdateFileCache();
+
+	for (FSection* Section : Sections)
+	{
+		if (Section)
+		{
+			Section->SetModel(this);
+		}
+	}
 #endif
 }
 
@@ -1615,7 +1618,7 @@ const FMLDeformerGeomCacheTrainingInputAnim* UNearestNeighborModel::GetNearestNe
 void UNearestNeighborModel::UpdatePCACoeffStarts()
 {
 	uint32 Acc = 0;
-	PCACoeffStarts.Reserve(Sections.Num());
+	PCACoeffStarts.Reset(Sections.Num());
 	for(const FSection* Section : Sections)
 	{
 		PCACoeffStarts.Add(Acc);
diff --git a/Engine/Plugins/Animation/MLDeformer/NearestNeighborModel/Source/NearestNeighborModel/Public/NearestNeighborModel.h b/Engine/Plugins/Animation/MLDeformer/NearestNeighborModel/Source/NearestNeighborModel/Public/NearestNeighborModel.h
--- a/Engine/Plugins/Animation/MLDeformer/NearestNeighborModel/Source/NearestNeighborModel/Public/NearestNeighborModel.h
+++ b/Engine/Plugins/Animation/MLDeformer/NearestNeighborModel/Source/NearestNeighborModel/Public/NearestNeighborModel.h
@@ -597,8 +597,8 @@ PRAGMA_ENABLE_DEPRECATION_WARNINGS
 	UPROPERTY(EditAnywhere, BlueprintReadWrite, Category = "Nearest Neighbor Settings")
 	bool bUsePCA = false;
 
-	/** The number of basis used in each section. Only editable when UsePCA is false. */
-	UPROPERTY(EditAnywhere, BlueprintReadWrite, Category = "Nearest Neighbor Settings", meta = (EditCondition = "!bUsePCA"))
+	/** The number of basis used in each section. */
+	UPROPERTY(EditAnywhere, BlueprintReadWrite, Category = "Nearest Neighbor Settings", meta = (EditCondition = "!bUsePCA", EditConditionHides))
 	int32 NumBasisPerSection = 128;
 
 	/** Whether to use dual quaternion deltas. If false, LBS deltas will be used. */
