[MLDeformer] Automatically create two hidden layers when creating a new Nearest Neighbor model. And automatically create a section as well.

diff --git a/Engine/Plugins/Animation/MLDeformer/NearestNeighborModel/Source/NearestNeighborModel/Private/NearestNeighborModel.cpp b/Engine/Plugins/Animation/MLDeformer/NearestNeighborModel/Source/NearestNeighborModel/Private/NearestNeighborModel.cpp
--- a/Engine/Plugins/Animation/MLDeformer/NearestNeighborModel/Source/NearestNeighborModel/Private/NearestNeighborModel.cpp
+++ b/Engine/Plugins/Animation/MLDeformer/NearestNeighborModel/Source/NearestNeighborModel/Private/NearestNeighborModel.cpp
@@ -295,34 +295,6 @@ namespace UE::NearestNeighborModel
 		}
 
 #if WITH_EDITORONLY_DATA
-		TArray<FInt32Range> GetMeshVertRanges(const USkeletalMesh& SkelMesh)
-		{
-			constexpr int32 LODIndex = 0;
-
-			if (!SkelMesh.HasMeshDescription(LODIndex))
-			{
-				return {};
-			}
-
-			const FMeshDescription* MeshDescription = SkelMesh.GetMeshDescription(LODIndex);
-			const FSkeletalMeshConstAttributes MeshAttributes(*MeshDescription);
-
-			if (!MeshAttributes.HasSourceGeometryParts())
-			{
-				return {};
-			}
-
-			TArray<FInt32Range> MeshVertRanges;
-			MeshVertRanges.Reserve(MeshAttributes.GetNumSourceGeometryParts());
-			const FSkeletalMeshAttributesShared::FSourceGeometryPartVertexOffsetAndCountConstRef PartOffsetAndCountRef = MeshAttributes.GetSourceGeometryPartVertexOffsetAndCounts();
-			for (const FSourceGeometryPartID GeometryPartID: MeshAttributes.SourceGeometryParts().GetElementIDs())
-			{
-				TArrayView<const int32> OffsetAndCount = PartOffsetAndCountRef.Get(GeometryPartID);
-				MeshVertRanges.Emplace(OffsetAndCount[0], OffsetAndCount[0] + OffsetAndCount[1]);
-			}
-			return MeshVertRanges;
-		}
-
 		TArray<FName> GetVertexFloatAttributeNames(const USkeletalMesh& SkelMesh)
 		{
 			constexpr int32 LODIndex = 0;
@@ -1168,6 +1140,10 @@ UNearestNeighborModel::UNearestNeighborModel(const FObjectInitializer& ObjectIni
 {
 #if WITH_EDITORONLY_DATA
 	SetVizSettings(ObjectInitializer.CreateEditorOnlyDefaultSubobject<UNearestNeighborModelVizSettings>(this, TEXT("VizSettings")));
+
+	// Add two layers of 128 units on default.
+	HiddenLayerDims.Add(128);
+	HiddenLayerDims.Add(128);
 #endif
 }
 
@@ -1236,7 +1212,7 @@ void UNearestNeighborModel::PostEditChangeProperty(FPropertyChangedEvent& Proper
 			FString MapString;
 			if (GetSkeletalMesh())
 			{
-				const TArray<FInt32Range> Ranges = UE::NearestNeighborModel::Private::GetMeshVertRanges(*GetSkeletalMesh());
+				const TArray<FInt32Range> Ranges = GetMeshVertRanges(*GetSkeletalMesh());
 				const int32 MeshIndex = Section->GetMeshIndex(); 
 				if (Ranges.IsValidIndex(MeshIndex))
 				{
@@ -1421,6 +1397,11 @@ const UNearestNeighborModelSection& UNearestNeighborModel::GetSection(int32 Inde
 	return *Sections[Index];
 }
 
+UNearestNeighborModelSection& UNearestNeighborModel::GetSection(int32 Index)
+{
+	return *Sections[Index];
+}
+
 const TArray<int32>& UNearestNeighborModel::GetPCACoeffStarts() const
 {
 	return PCACoeffStarts;
@@ -1516,6 +1497,38 @@ float UNearestNeighborModel::GetRBFSigma() const
 	return RBFSigma;
 }
 
+
+#if WITH_EDITORONLY_DATA
+TArray<FInt32Range> UNearestNeighborModel::GetMeshVertRanges(const USkeletalMesh& SkelMesh)
+{
+	constexpr int32 LODIndex = 0;
+
+	if (!SkelMesh.HasMeshDescription(LODIndex))
+	{
+		return {};
+	}
+
+	const FMeshDescription* MeshDescription = SkelMesh.GetMeshDescription(LODIndex);
+	const FSkeletalMeshConstAttributes MeshAttributes(*MeshDescription);
+
+	if (!MeshAttributes.HasSourceGeometryParts())
+	{
+		return {};
+	}
+
+	TArray<FInt32Range> MeshVertRanges;
+	MeshVertRanges.Reserve(MeshAttributes.GetNumSourceGeometryParts());
+	const FSkeletalMeshAttributesShared::FSourceGeometryPartVertexOffsetAndCountConstRef PartOffsetAndCountRef = MeshAttributes.GetSourceGeometryPartVertexOffsetAndCounts();
+	for (const FSourceGeometryPartID GeometryPartID: MeshAttributes.SourceGeometryParts().GetElementIDs())
+	{
+		TArrayView<const int32> OffsetAndCount = PartOffsetAndCountRef.Get(GeometryPartID);
+		MeshVertRanges.Emplace(OffsetAndCount[0], OffsetAndCount[0] + OffsetAndCount[1]);
+	}
+	return MeshVertRanges;
+}
+#endif
+
+
 #if WITH_EDITOR
 UNearestNeighborModelSection* UNearestNeighborModel::OnSectionAdded(int32 NewIndex)
 {
@@ -1537,11 +1550,6 @@ UNearestNeighborModelSection* UNearestNeighborModel::OnSectionAdded(int32 NewInd
 	return Section;
 }
 
-UNearestNeighborModelSection& UNearestNeighborModel::GetSection(int32 Index)
-{
-	return *Sections[Index];
-}
-
 FDateTime UNearestNeighborModel::GetNetworkLastWriteTime() const
 {
 	return NetworkLastWriteTime;
@@ -1972,7 +1980,16 @@ bool UNearestNeighborModel::LoadOptimizedNetworkFromFile(const FString& Filename
 	}
 	return false;
 }
-#endif
+
+void UNearestNeighborModel::RemoveAllSections()
+{
+	for (TObjectPtr<UNearestNeighborModelSection>& Section : Sections)
+	{
+		Section->ConditionalBeginDestroy();
+	}
+	Sections.Empty();
+}
+#endif	// WITH_EDITOR
 
 int32 UNearestNeighborModel::GetNumNetworkOutputs() const
 {
diff --git a/Engine/Plugins/Animation/MLDeformer/NearestNeighborModel/Source/NearestNeighborModel/Public/NearestNeighborModel.h b/Engine/Plugins/Animation/MLDeformer/NearestNeighborModel/Source/NearestNeighborModel/Public/NearestNeighborModel.h
--- a/Engine/Plugins/Animation/MLDeformer/NearestNeighborModel/Source/NearestNeighborModel/Public/NearestNeighborModel.h
+++ b/Engine/Plugins/Animation/MLDeformer/NearestNeighborModel/Source/NearestNeighborModel/Public/NearestNeighborModel.h
@@ -379,6 +379,7 @@ public:
 	UFUNCTION(BlueprintPure, Category = "Python")
 	const UNearestNeighborModelSection* GetSectionPtr(int32 Index) const;
 	const UNearestNeighborModelSection& GetSection(int32 Index) const;
+	UNearestNeighborModelSection& GetSection(int32 Index);
 
 	UFUNCTION(BlueprintPure, Category = "Nearest Neighbor Model")
 	const TArray<int32>& GetPCACoeffStarts() const;
@@ -407,6 +408,10 @@ public:
 	UFUNCTION(BlueprintPure, Category = "Nearest Neighbor Model")
 	bool DoesUsePCA() const { return bUsePCA; }
 
+#if WITH_EDITORONLY_DATA
+	TArray<FInt32Range> GetMeshVertRanges(const USkeletalMesh& SkelMesh);
+#endif
+
 #if WITH_EDITOR
 	UFUNCTION(BlueprintPure, Category = "Nearest Neighbor Model")
 	int32 GetInputDim() const { return InputDim; }
@@ -471,10 +476,13 @@ public:
 	void InvalidateInferenceModelOnly();
 	bool LoadOptimizedNetworkFromFile(const FString& Filename);
 	void ClearOptimizedNetwork();
+	void RemoveAllSections();
 
 	FMLDeformerGeomCacheTrainingInputAnim* GetNearestNeighborAnim(int32 SectionIndex);
 	const FMLDeformerGeomCacheTrainingInputAnim* GetNearestNeighborAnim(int32 SectionIndex) const;
 
+	void AddSection(TObjectPtr<UNearestNeighborModelSection> Section) { Sections.Add(Section); }
+
 	void UpdateFileCache();
 	const FString& GetFileCacheDirectory() const;
 	void SetFileCacheDirectory(const FString& InFileCacheDirectory);
@@ -621,12 +629,15 @@ PRAGMA_ENABLE_DEPRECATION_WARNINGS
 	UPROPERTY(EditAnywhere, BlueprintReadWrite, Category = "Nearest Neighbor Settings", META = (ClampMin = "0"))
 	float RBFSigma = 1.0f;
 
+public:
+#if WITH_EDITOR
+	UNearestNeighborModelSection* OnSectionAdded(int32 NewIndex);
+#endif
+
 private:
 	TWeakObjectPtr<UNetwork> GetOptimizedNetwork();
 
 #if WITH_EDITOR
-	UNearestNeighborModelSection* OnSectionAdded(int32 NewIndex);
-	FSection& GetSection(int32 Index);
 
 	void SetOptimizedNetwork(UNearestNeighborOptimizedNetwork* InOptimizedNetwork);
 
diff --git a/Engine/Plugins/Animation/MLDeformer/NearestNeighborModel/Source/NearestNeighborModelEditor/Private/NearestNeighborEditorModel.cpp b/Engine/Plugins/Animation/MLDeformer/NearestNeighborModel/Source/NearestNeighborModelEditor/Private/NearestNeighborEditorModel.cpp
--- a/Engine/Plugins/Animation/MLDeformer/NearestNeighborModel/Source/NearestNeighborModelEditor/Private/NearestNeighborEditorModel.cpp
+++ b/Engine/Plugins/Animation/MLDeformer/NearestNeighborModel/Source/NearestNeighborModelEditor/Private/NearestNeighborEditorModel.cpp
@@ -58,11 +58,18 @@ namespace UE::NearestNeighborModel
 		InitInputInfo(Model->GetInputInfo());
 		VertexMapSelector = MakeUnique<FVertexMapSelector>();
 		VertexMapSelector->Update(Model->GetSkeletalMesh());
-		const UNearestNeighborModel* const NearestNeighborModel = GetCastModel();
+		UNearestNeighborModel* NearestNeighborModel = GetCastModel();
 		if (!NearestNeighborModel)
 		{
 			return;
 		}
+
+		// Create one section on default.
+		if (NearestNeighborModel->GetNumSections() == 0)
+		{
+			CreateDefaultSection();
+		}
+
 		const int32 NumSections = NearestNeighborModel->GetNumSections();
 		if (UNearestNeighborModelVizSettings* const NNViz = GetCastVizSettings())
 		{
@@ -820,6 +827,39 @@ namespace UE::NearestNeighborModel
 		return nullptr;
 	}
 
+	void FNearestNeighborEditorModel::CreateDefaultSection()
+	{
+		UNearestNeighborModel* NearestNeighborModel = GetCastModel();
+		VertexMapSelector->Update(Model->GetSkeletalMesh());
+		NearestNeighborModel->AddSection(nullptr);
+		FSection* const Section = NearestNeighborModel->OnSectionAdded(0);
+		FString MapString;
+		if (NearestNeighborModel->GetSkeletalMesh())
+		{
+			const TArray<FInt32Range> Ranges = NearestNeighborModel->GetMeshVertRanges(*NearestNeighborModel->GetSkeletalMesh());
+			const int32 MeshIndex = Section->GetMeshIndex(); 
+			if (Ranges.IsValidIndex(MeshIndex))
+			{
+				MapString = FString::Printf(TEXT("%d-%d"), Ranges[MeshIndex].GetLowerBoundValue(), Ranges[MeshIndex].GetUpperBoundValue() - 1);
+			}
+		}
+		Section->SetVertexMapString(MapString);
+		NearestNeighborModel->InvalidateTrainingModelOnly();
+		NearestNeighborModel->UpdateNetworkOutputDim();
+		NearestNeighborModel->InvalidateTraining();
+	}
+
+	void FNearestNeighborEditorModel::CopyBaseSettingsFromModel(const FMLDeformerEditorModel* SourceEditorModel)
+	{
+		FMLDeformerMorphModelEditorModel::CopyBaseSettingsFromModel(SourceEditorModel);
+
+		UNearestNeighborModel* NearestNeighborModel = GetCastModel();
+
+		// Wipe existing sections (as our init might have created one) and create a default new one.
+		NearestNeighborModel->RemoveAllSections();
+		CreateDefaultSection();
+	}
+
 	void FVertexMapSelector::Update(const USkeletalMesh* SkelMesh)
 	{
 		constexpr int32 LODIndex = 0;
diff --git a/Engine/Plugins/Animation/MLDeformer/NearestNeighborModel/Source/NearestNeighborModelEditor/Public/NearestNeighborEditorModel.h b/Engine/Plugins/Animation/MLDeformer/NearestNeighborModel/Source/NearestNeighborModelEditor/Public/NearestNeighborEditorModel.h
--- a/Engine/Plugins/Animation/MLDeformer/NearestNeighborModel/Source/NearestNeighborModelEditor/Public/NearestNeighborEditorModel.h
+++ b/Engine/Plugins/Animation/MLDeformer/NearestNeighborModel/Source/NearestNeighborModelEditor/Public/NearestNeighborEditorModel.h
@@ -49,6 +49,7 @@ namespace UE::NearestNeighborModel
 		virtual void OnPostTraining(ETrainingResult TrainingResult, bool bUsePartiallyTrainedWhenAborted) override;
 		virtual void Render(const FSceneView* View, FViewport* Viewport, FPrimitiveDrawInterface* PDI) override;
 		virtual void UpdateTrainingDeviceList() override;
+		virtual void CopyBaseSettingsFromModel(const FMLDeformerEditorModel* SourceEditorModel) override;
 		// ~END FMLDeformerEditorModel overrides.
 		
 		// UMLDeformerMorphModelEditorModel overrides.
@@ -74,6 +75,7 @@ namespace UE::NearestNeighborModel
 		UMLDeformerComponent* GetTestMLDeformerComponent() const;
 		USkeletalMeshComponent* GetTestSkeletalMeshComponent() const;
 		UNearestNeighborModelInstance* GetTestNearestNeighborModelInstance() const;
+		void CreateDefaultSection();
 
 		EOpFlag Update();
 		EOpFlag CheckNetwork();
