NearestNeighbor Model: set bUsePCA = true before doing UpdateForTraining when upgrading old assets.

Without this fix, the PCACoeffStarts would be empty, which would cause things like Key Pose Extractor tool to fail.

diff --git a/Engine/Plugins/Animation/MLDeformer/NearestNeighborModel/Source/NearestNeighborModel/Private/NearestNeighborModel.cpp b/Engine/Plugins/Animation/MLDeformer/NearestNeighborModel/Source/NearestNeighborModel/Private/NearestNeighborModel.cpp
--- a/Engine/Plugins/Animation/MLDeformer/NearestNeighborModel/Source/NearestNeighborModel/Private/NearestNeighborModel.cpp
+++ b/Engine/Plugins/Animation/MLDeformer/NearestNeighborModel/Source/NearestNeighborModel/Private/NearestNeighborModel.cpp
@@ -1351,6 +1351,7 @@ void UNearestNeighborModel::PostLoad()
 			check(Section);
 			Section->InitFromClothPartData(ClothPartData_DEPRECATED[Index]);
 		}
+		bUsePCA = true;
 		UpdateForTraining();
 		UpdateForInference();
 	PRAGMA_ENABLE_DEPRECATION_WARNINGS
