NearestNeighbor Model: Fix a crash on editor exit that occurs after training.

Crash was caused by not correctly releasing the morph target buffers inside OnPostTraining. This is already done by the base class, so no need to even do anything here.

Also removed the code to restore backup morph target data since this is also done in the base class.

diff --git a/Engine/Plugins/Animation/MLDeformer/NearestNeighborModel/Source/NearestNeighborModelEditor/Private/NearestNeighborEditorModel.cpp b/Engine/Plugins/Animation/MLDeformer/NearestNeighborModel/Source/NearestNeighborModelEditor/Private/NearestNeighborEditorModel.cpp
--- a/Engine/Plugins/Animation/MLDeformer/NearestNeighborModel/Source/NearestNeighborModelEditor/Private/NearestNeighborEditorModel.cpp
+++ b/Engine/Plugins/Animation/MLDeformer/NearestNeighborModel/Source/NearestNeighborModelEditor/Private/NearestNeighborEditorModel.cpp
@@ -184,15 +184,9 @@ namespace UE::NearestNeighborModel
 
 	void FNearestNeighborEditorModel::OnPostTraining(ETrainingResult TrainingResult, bool bUsePartiallyTrainedWhenAborted)
 	{
-		UNearestNeighborModel* const NearestNeighborModel = GetCastModel();
-		if (TrainingResult == ETrainingResult::Aborted && !bUsePartiallyTrainedWhenAborted)
-		{
-			GetMorphModel()->SetMorphTargetDeltas(MorphTargetDeltasBackup);
-			GetMorphModel()->SetMorphTargetsMinMaxWeights(MorphTargetsMinMaxWeightsBackup);
-		}
-		else if (TrainingResult == ETrainingResult::Success || (TrainingResult == ETrainingResult::Aborted && bUsePartiallyTrainedWhenAborted))
+		if (UNearestNeighborModel* const NearestNeighborModel = GetCastModel())
 		{
-			if (NearestNeighborModel)
+			if (TrainingResult == ETrainingResult::Success || (TrainingResult == ETrainingResult::Aborted && bUsePartiallyTrainedWhenAborted))
 			{
 				NearestNeighborModel->InvalidateInference();
 				if (NearestNeighborModel->DoesUseFileCache())
@@ -200,10 +194,6 @@ namespace UE::NearestNeighborModel
 					NearestNeighborModel->UpdateFileCache();
 				}
 			}
-			ResetMorphTargets();
-		}
-		if (NearestNeighborModel)
-		{
 			NearestNeighborModel->UpdateForInference();
 		}
 		FMLDeformerMorphModelEditorModel::OnPostTraining(TrainingResult, bUsePartiallyTrainedWhenAborted);
@@ -765,19 +755,6 @@ namespace UE::NearestNeighborModel
 		return Result;
 	}
 
-	void FNearestNeighborEditorModel::ResetMorphTargets()
-	{
-		const int32 NumLODs = GetMorphModel()->GetNumLODs();
-		for (int32 LOD = 0; LOD < NumLODs; ++LOD)
-		{
-			const TSharedPtr<FExternalMorphSet> MorphSet = GetMorphModel()->GetMorphTargetSet(LOD);
-			if (MorphSet.IsValid())
-			{
-				MorphSet->MorphBuffers = FMorphTargetVertexInfoBuffers();
-			}
-		}
-	}
-
 	void FNearestNeighborEditorModel::UpdateNearestNeighborIds()
 	{
 		const UNearestNeighborModelInstance* ModelInstance = GetTestNearestNeighborModelInstance();
diff --git a/Engine/Plugins/Animation/MLDeformer/NearestNeighborModel/Source/NearestNeighborModelEditor/Public/NearestNeighborEditorModel.h b/Engine/Plugins/Animation/MLDeformer/NearestNeighborModel/Source/NearestNeighborModelEditor/Public/NearestNeighborEditorModel.h
--- a/Engine/Plugins/Animation/MLDeformer/NearestNeighborModel/Source/NearestNeighborModelEditor/Public/NearestNeighborEditorModel.h
+++ b/Engine/Plugins/Animation/MLDeformer/NearestNeighborModel/Source/NearestNeighborModelEditor/Public/NearestNeighborEditorModel.h
@@ -81,7 +81,6 @@ namespace UE::NearestNeighborModel
 		EOpFlag CheckNetwork();
 		EOpFlag UpdateNearestNeighborData();
 		EOpFlag UpdateMorphDeltas();
-		void ResetMorphTargets();
 		void UpdateNearestNeighborIds();
 
 		FNearestNeighborEditorModelActor* CreateNearestNeighborActor(UWorld* World) const;
