NearestNeighborModel: hide caching options behind a cvar.

Caching only works under very specific scenarios and needs to be reworked as part of a larger NNM refactor. Since it can be useful in very limited scenarios, keeping it around, but not by default.

Caching options are available if caching was already enabled (so you can uncheck the box without knowing about the cvar).

diff --git a/Engine/Plugins/Animation/MLDeformer/NearestNeighborModel/Source/NearestNeighborModel/Public/NearestNeighborModel.h b/Engine/Plugins/Animation/MLDeformer/NearestNeighborModel/Source/NearestNeighborModel/Public/NearestNeighborModel.h
--- a/Engine/Plugins/Animation/MLDeformer/NearestNeighborModel/Source/NearestNeighborModel/Public/NearestNeighborModel.h
+++ b/Engine/Plugins/Animation/MLDeformer/NearestNeighborModel/Source/NearestNeighborModel/Public/NearestNeighborModel.h
@@ -568,7 +568,7 @@ protected:
 	bool bUseFileCache = false;
 
 	/** Directory to save the intermediate results. */
-	UPROPERTY(VisibleAnywhere, BlueprintReadOnly, Category = "File Cache", meta = (EditCondition = "bUseFileCache"))
+	UPROPERTY(VisibleAnywhere, BlueprintReadOnly, Category = "File Cache", meta = (EditCondition = "bUseFileCache", EditConditionHides))
 	FString FileCacheDirectory = FPaths::ProjectIntermediateDir() + "NearestNeighborModel";
 
 	TOptional<FDateTime> CachedDeltasTimestamp;
diff --git a/Engine/Plugins/Animation/MLDeformer/NearestNeighborModel/Source/NearestNeighborModelEditor/Private/NearestNeighborModelDetails.cpp b/Engine/Plugins/Animation/MLDeformer/NearestNeighborModel/Source/NearestNeighborModelEditor/Private/NearestNeighborModelDetails.cpp
--- a/Engine/Plugins/Animation/MLDeformer/NearestNeighborModel/Source/NearestNeighborModelEditor/Private/NearestNeighborModelDetails.cpp
+++ b/Engine/Plugins/Animation/MLDeformer/NearestNeighborModel/Source/NearestNeighborModelEditor/Private/NearestNeighborModelDetails.cpp
@@ -20,8 +20,13 @@
 
 #define LOCTEXT_NAMESPACE "NearestNeighborModelDetails"
 
-namespace UE::NearestNeighborModel
+namespace UE::NearestNeighborModel	
 {
+
+	static bool bShowFileCacheCategory = false;
+	FAutoConsoleVariableRef CVarNearestNeighborShowFileCacheCategory(TEXT("MLDeformer.NearestNeighbor.ShowFileCacheCategory"), bShowFileCacheCategory, TEXT("Whether or not to show the file cache category in the NearestNeighbor details panel"));
+
+
 	TSharedRef<IDetailCustomization> FNearestNeighborModelDetails::MakeInstance()
 	{
 		return MakeShareable(new FNearestNeighborModelDetails());
@@ -92,47 +97,48 @@ namespace UE::NearestNeighborModel
 				];
 		}
 
-		void AddFileCacheRow(IDetailCategoryBuilder& CategoryBuilder, const FText& FilterString, const FText& Name, const FGetTextDelegate ValueDelegate, TSharedRef<SWidget>& ValueWidget, TAttribute<bool> bIsEnabled)
+		void AddFileCacheRow(IDetailCategoryBuilder& CategoryBuilder, const FText& FilterString, const FText& Name, const FGetTextDelegate ValueDelegate, TSharedRef<SWidget>& ValueWidget, TAttribute<bool> bIsEnabled, TAttribute<EVisibility> IsVisible)
 		{
 			check(ValueDelegate.IsBound());
 			CategoryBuilder.AddCustomRow(FilterString)
 				.NameContent()
 				[
 					SNew(SHorizontalBox)
-					+SHorizontalBox::Slot()
-					.Padding(2, 2)
-					.VAlign(VAlign_Center)
-					[
-						SNew(STextBlock)
-						.Text(Name)
-						.Font(IDetailLayoutBuilder::GetDetailFont())
-					]
+						+ SHorizontalBox::Slot()
+						.Padding(2, 2)
+						.VAlign(VAlign_Center)
+						[
+							SNew(STextBlock)
+								.Text(Name)
+								.Font(IDetailLayoutBuilder::GetDetailFont())
+						]
 				]
 				.ValueContent()
 				.MinDesiredWidth(300.f)
 				.MaxDesiredWidth(300.f)
 				[
 					SNew(SHorizontalBox)
-					+SHorizontalBox::Slot()
-					.Padding(2, 2)
-					.VAlign(VAlign_Center)
-					.FillWidth(0.6f)
-					[
-						SNew(STextBlock)
-						.Text_Lambda([ValueDelegate]
-						{
-							return ValueDelegate.Execute();
-						})
-						.Font(IDetailLayoutBuilder::GetDetailFont())
-					]
-					+SHorizontalBox::Slot()
-					.FillWidth(0.4f)
-					.VAlign(VAlign_Center)
-					[
-						ValueWidget
-					]
+						+ SHorizontalBox::Slot()
+						.Padding(2, 2)
+						.VAlign(VAlign_Center)
+						.FillWidth(0.6f)
+						[
+							SNew(STextBlock)
+								.Text_Lambda([ValueDelegate]
+									{
+										return ValueDelegate.Execute();
+									})
+								.Font(IDetailLayoutBuilder::GetDetailFont())
+						]
+						+ SHorizontalBox::Slot()
+						.FillWidth(0.4f)
+						.VAlign(VAlign_Center)
+						[
+							ValueWidget
+						]
 				]
-				.IsEnabled(bIsEnabled);
+				.IsEnabled(bIsEnabled)
+				.Visibility(IsVisible);
 		}
 
 		FDateTime ToLocalTime(const FDateTime& UtcTime)
@@ -410,13 +416,18 @@ namespace UE::NearestNeighborModel
 	{
 		IDetailCategoryBuilder& Builder = DetailLayoutBuilder->EditCategory("File Cache", FText::GetEmpty(), ECategoryPriority::Important);
 
-		Builder.AddProperty(UNearestNeighborModel::GetUseFileCachePropertyName());
+		TAttribute<EVisibility> IsVisible = TAttribute<EVisibility>::Create([this]
+			{
+				UNearestNeighborModel* const NearestNeighborModel = GetCastModel();
+				return ((NearestNeighborModel && NearestNeighborModel->DoesUseFileCache()) || bShowFileCacheCategory) ? EVisibility::Visible : EVisibility::Collapsed;
+			});
+
+
+		Builder.AddProperty(UNearestNeighborModel::GetUseFileCachePropertyName())
+			.Visibility(IsVisible);
+
 		Builder.AddCustomRow(FText::FromString("FileCacheWarning"))
-		.Visibility(TAttribute<EVisibility>::Create([this]
-		{
-			const UNearestNeighborModel* const NearestNeighborModel = GetCastModel();
-			return NearestNeighborModel && NearestNeighborModel->DoesUseFileCache() ? EVisibility::Visible : EVisibility::Collapsed;
-		}))
+		.Visibility(IsVisible)
 		.WholeRowContent()
 		[
 			SNew(SBox)
@@ -524,7 +535,8 @@ namespace UE::NearestNeighborModel
 					return NearestNeighborModel->GetCachedDeltasTimestamp();
 				}),
 				Widget,
-				bIsEnabled);
+				bIsEnabled,
+				IsVisible);
 		}
 		{
 			TSharedRef<SWidget> Widget = CreateDeleteButton([NearestNeighborModel]()
@@ -539,7 +551,8 @@ namespace UE::NearestNeighborModel
 					return NearestNeighborModel->GetCachedPCATimestamp();
 				}),
 				Widget,
-				bIsEnabled);
+				bIsEnabled,
+				IsVisible);
 		}
 		{
 			TSharedRef<SWidget> Widget = CreateDeleteButton([NearestNeighborModel]()
@@ -554,7 +567,8 @@ namespace UE::NearestNeighborModel
 					return NearestNeighborModel->GetCachedNetworkTimestamp();
 				}),
 				Widget,
-				bIsEnabled);
+				bIsEnabled,
+				IsVisible);
 		}
 
 		Builder.AddCustomRow(FText::FromString("FileCacheRefresh")).WholeRowContent()
@@ -572,6 +586,7 @@ namespace UE::NearestNeighborModel
 					return FReply::Handled();
 				})
 				.IsEnabled(bIsEnabled)
+				.Visibility(IsVisible)
 			]
 		];
 	}
