[MLDeformer] Fixed a crash in the Nearest Neighbor Model where the GetNumBaseMeshVerts() is 0, which can be caused by the imported model data not to be present.

diff --git a/Engine/Plugins/Animation/MLDeformer/NearestNeighborModel/Source/NearestNeighborModel/Private/NearestNeighborModel.cpp b/Engine/Plugins/Animation/MLDeformer/NearestNeighborModel/Source/NearestNeighborModel/Private/NearestNeighborModel.cpp
--- a/Engine/Plugins/Animation/MLDeformer/NearestNeighborModel/Source/NearestNeighborModel/Private/NearestNeighborModel.cpp
+++ b/Engine/Plugins/Animation/MLDeformer/NearestNeighborModel/Source/NearestNeighborModel/Private/NearestNeighborModel.cpp
@@ -1773,6 +1773,11 @@ const TArray<float>& UNearestNeighborModel::GetVertexWeightSum() const
 
 void UNearestNeighborModel::NormalizeVertexWeights()
 {
+	if (GetNumBaseMeshVerts() == 0)
+	{
+		return;
+	}
+
 	VertexWeightSum.Init(0.f, GetNumBaseMeshVerts());
 	for (const FSection* Section : Sections)
 	{
