[MLDeformer] Removed an error log that shouldn't really be logged in the Nearest Neighbor model.

diff --git a/Engine/Plugins/Animation/MLDeformer/NearestNeighborModel/Source/NearestNeighborModel/Private/NearestNeighborModel.cpp b/Engine/Plugins/Animation/MLDeformer/NearestNeighborModel/Source/NearestNeighborModel/Private/NearestNeighborModel.cpp
--- a/Engine/Plugins/Animation/MLDeformer/NearestNeighborModel/Source/NearestNeighborModel/Private/NearestNeighborModel.cpp
+++ b/Engine/Plugins/Animation/MLDeformer/NearestNeighborModel/Source/NearestNeighborModel/Private/NearestNeighborModel.cpp
@@ -824,7 +824,6 @@ UE::NearestNeighborModel::EOpFlag UNearestNeighborModelSection::UpdateVertexWeig
 	const USkeletalMesh* const SkeletalMesh = GetModel()->GetSkeletalMesh();
 	if (!SkeletalMesh)
 	{
-		UE_LOG(LogNearestNeighborModel, Error, TEXT("SkeletalMesh is None"));
 		return EOpFlag::Error;
 	}
 	if (!SkeletalMesh->GetImportedModel())
@@ -874,7 +873,6 @@ UE::NearestNeighborModel::EOpFlag UNearestNeighborModelSection::UpdateVertexWeig
 	const USkeletalMesh* const SkeletalMesh = GetModel()->GetSkeletalMesh();
 	if (!SkeletalMesh)
 	{
-		UE_LOG(LogNearestNeighborModel, Error, TEXT("SkeletalMesh is None"));
 		return EOpFlag::Error;
 	}
 	if (!SkeletalMesh->GetImportedModel())
@@ -962,7 +960,6 @@ UE::NearestNeighborModel::EOpFlag UNearestNeighborModelSection::UpdateVertexWeig
 	const USkeletalMesh* const SkeletalMesh = GetModel()->GetSkeletalMesh();
 	if (!SkeletalMesh)
 	{
-		UE_LOG(LogNearestNeighborModel, Error, TEXT("SkeletalMesh is None"));
 		return EOpFlag::Error;
 	}
 
@@ -1023,7 +1020,6 @@ UE::NearestNeighborModel::EOpFlag UNearestNeighborModelSection::UpdateVertexWeig
 	const USkeletalMesh* const SkeletalMesh = GetModel()->GetSkeletalMesh();
 	if (!SkeletalMesh)
 	{
-		UE_LOG(LogNearestNeighborModel, Error, TEXT("SkeletalMesh is None"));
 		return EOpFlag::Error;
 	}
 	if (!GetModel())
