[MLDeformer] Fixed a crash when using dual quaternions. Somehow the render data of the SkeletalMeshObject can be a nullptr. This fix handles that case.

#rnx

diff --git a/Engine/Plugins/Animation/MLDeformer/MLDeformerFramework/Source/MLDeformerFramework/Private/DeformerGraphDataInterfaceSkeletonWithQuats.cpp b/Engine/Plugins/Animation/MLDeformer/MLDeformerFramework/Source/MLDeformerFramework/Private/DeformerGraphDataInterfaceSkeletonWithQuats.cpp
--- a/Engine/Plugins/Animation/MLDeformer/MLDeformerFramework/Source/MLDeformerFramework/Private/DeformerGraphDataInterfaceSkeletonWithQuats.cpp
+++ b/Engine/Plugins/Animation/MLDeformer/MLDeformerFramework/Source/MLDeformerFramework/Private/DeformerGraphDataInterfaceSkeletonWithQuats.cpp
@@ -224,6 +224,11 @@ namespace UE::MLDeformer
 		TArray<FQuat4f> RefToLocalQuats;
 		CacheRefToLocalQuats(RefToLocalQuats);
 
+		if (RefToLocalQuats.IsEmpty())
+		{
+			return;
+		}
+
 		const int32 LodIndex = SkeletalMeshObject->GetLOD();
 		FSkeletalMeshRenderData const& SkeletalMeshRenderData = SkeletalMeshObject->GetSkeletalMeshRenderData();
 		FSkeletalMeshLODRenderData const* LodRenderData = &SkeletalMeshRenderData.LODRenderData[LodIndex];
@@ -288,7 +293,7 @@ namespace UE::MLDeformer
 			Parameters.InputWeightStride = WeightBuffer->GetConstantInfluencesVertexStride();
 			Parameters.InputWeightIndexSize = WeightBuffer->GetBoneIndexByteSize() | (WeightBuffer->GetBoneWeightByteSize() << 8);
 			Parameters.BoneMatrices = BoneBufferSRV != nullptr ? BoneBufferSRV : NullSRVBinding;
-			Parameters.BoneQuats = PerSectionRefToLocalQuatsSRVs[InvocationIndex];
+			Parameters.BoneQuats = PerSectionRefToLocalQuatsSRVs.IsValidIndex(InvocationIndex) ? PerSectionRefToLocalQuatsSRVs[InvocationIndex] : nullptr;
 			Parameters.InputWeightStream = SkinWeightBufferSRV != nullptr ? SkinWeightBufferSRV : NullSRVBinding;
 			Parameters.InputWeightLookupStream = InputWeightLookupStreamSRV != nullptr ? InputWeightLookupStreamSRV : NullSRVBinding;
 		}
@@ -297,7 +302,7 @@ namespace UE::MLDeformer
 	void FDeformerGraphSkeletonWithWeightedQuatsDataProviderProxy::CacheRefToLocalQuats(TArray<FQuat4f>& OutRefToLocalQuats) const
 	{
 		OutRefToLocalQuats.Reset();
-		if (SkeletalMeshObject == nullptr)
+		if (!SkeletalMeshObject || !SkeletalMeshObject->HaveValidDynamicData())
 		{
 			return;
 		}
diff --git a/Engine/Plugins/Animation/MLDeformer/NearestNeighborModel/Source/NearestNeighborModel/Private/OptimusDataInterfaceSkeletonWithQuats.cpp b/Engine/Plugins/Animation/MLDeformer/NearestNeighborModel/Source/NearestNeighborModel/Private/OptimusDataInterfaceSkeletonWithQuats.cpp
--- a/Engine/Plugins/Animation/MLDeformer/NearestNeighborModel/Source/NearestNeighborModel/Private/OptimusDataInterfaceSkeletonWithQuats.cpp
+++ b/Engine/Plugins/Animation/MLDeformer/NearestNeighborModel/Source/NearestNeighborModel/Private/OptimusDataInterfaceSkeletonWithQuats.cpp
@@ -228,6 +228,11 @@ void FOptimusSkeletonWithQuatsDataProviderProxy::AllocateResources(FRDGBuilder&
 	TArray<FQuat4f> RefToLocalQuats;
 	CacheRefToLocalQuats(RefToLocalQuats);
 
+	if (RefToLocalQuats.IsEmpty())
+	{
+		return;
+	}
+
 	const int32 LodIndex = SkeletalMeshObject->GetLOD();
 	FSkeletalMeshRenderData const& SkeletalMeshRenderData = SkeletalMeshObject->GetSkeletalMeshRenderData();
 	FSkeletalMeshLODRenderData const* LodRenderData = &SkeletalMeshRenderData.LODRenderData[LodIndex];
@@ -292,17 +297,16 @@ void FOptimusSkeletonWithQuatsDataProviderProxy::GatherDispatchData(FDispatchDat
 		Parameters.InputWeightStride = WeightBuffer->GetConstantInfluencesVertexStride();
 		Parameters.InputWeightIndexSize = WeightBuffer->GetBoneIndexByteSize() | (WeightBuffer->GetBoneWeightByteSize() << 8);
 		Parameters.BoneMatrices = BoneBufferSRV != nullptr ? BoneBufferSRV : NullSRVBinding;
-		Parameters.BoneQuats = PerSectionRefToLocalQuatsSRVs[InvocationIndex];
+		Parameters.BoneQuats = PerSectionRefToLocalQuatsSRVs.IsValidIndex(InvocationIndex) ? PerSectionRefToLocalQuatsSRVs[InvocationIndex] : nullptr;
 		Parameters.InputWeightStream = SkinWeightBufferSRV != nullptr ? SkinWeightBufferSRV : NullSRVBinding;
 		Parameters.InputWeightLookupStream = InputWeightLookupStreamSRV != nullptr ? InputWeightLookupStreamSRV : NullSRVBinding;
 	}
 }
 
-
 void FOptimusSkeletonWithQuatsDataProviderProxy::CacheRefToLocalQuats(TArray<FQuat4f>& OutRefToLocalQuats) const
 {
 	OutRefToLocalQuats.Reset();
-	if (SkeletalMeshObject == nullptr)
+	if (!SkeletalMeshObject || !SkeletalMeshObject->HaveValidDynamicData())
 	{
 		return;
 	}
